Renaming the AnimNextAnimationGraph to UAFAnimGraph in Unit test and in python test

#Synced-CL 50774406

diff --git a/Engine/Plugins/Experimental/UAF/UAFTestSuites/Source/UAFAnimGraphTestSuite/Private/AnimNextAnimGraphEditorVariablesTest.cpp b/Engine/Plugins/Experimental/UAF/UAFTestSuites/Source/UAFAnimGraphTestSuite/Private/AnimNextAnimGraphEditorVariablesTest.cpp
--- a/Engine/Plugins/Experimental/UAF/UAFTestSuites/Source/UAFAnimGraphTestSuite/Private/AnimNextAnimGraphEditorVariablesTest.cpp
+++ b/Engine/Plugins/Experimental/UAF/UAFTestSuites/Source/UAFAnimGraphTestSuite/Private/AnimNextAnimGraphEditorVariablesTest.cpp
@@ -235,7 +235,7 @@ bool FEditor_AnimGraph_Variables_Python::RunTest(const FString& InParameters)
 
 	const TCHAR* Script = TEXT(
 		"asset_tools = unreal.AssetToolsHelpers.get_asset_tools()\n"
-		"animation_graph = unreal.AssetTools.create_asset(asset_tools, asset_name = \"TestAnimGraph\", package_path = \"/Game/\", asset_class = unreal.AnimNextAnimationGraph, factory = unreal.AnimNextAnimationGraphFactory())\n"
+		"animation_graph = unreal.AssetTools.create_asset(asset_tools, asset_name = \"TestAnimGraph\", package_path = \"/Game/\", asset_class = unreal.UAFAnimGraph, factory = unreal.UAFAnimGraphFactory())\n"
 		"animation_graph.add_variable(name = \"TestParam\", value_type = unreal.PropertyBagPropertyType.BOOL, container_type = unreal.PropertyBagContainerType.NONE)\n"
 		"unreal.EditorAssetLibrary.delete_loaded_asset(animation_graph)\n"
 	);
