Add start time options to FUAFGraphFactoryAsset_Animation

#Synced-CL 50915596

diff --git a/Engine/Plugins/Experimental/UAF/UAFAnimGraph/Source/UAFAnimGraph/Private/AnimNextAnimGraphModule.cpp b/Engine/Plugins/Experimental/UAF/UAFAnimGraph/Source/UAFAnimGraph/Private/AnimNextAnimGraphModule.cpp
--- a/Engine/Plugins/Experimental/UAF/UAFAnimGraph/Source/UAFAnimGraph/Private/AnimNextAnimGraphModule.cpp
+++ b/Engine/Plugins/Experimental/UAF/UAFAnimGraph/Source/UAFAnimGraph/Private/AnimNextAnimGraphModule.cpp
@@ -183,6 +183,7 @@ void FAnimNextAnimGraphModule::RegisterGraphFactories()
 					InStruct.AnimSequence = AssetData.AnimationSequence;
 					InStruct.PlayRate = AssetData.PlayRate;
 					InStruct.LoopMode = AssetData.LoopMode;
+					InStruct.StartPosition = AssetData.CalculateStartTime();
 				}));
 			});
 	}
@@ -245,6 +246,9 @@ void FAnimNextAnimGraphModule::RegisterSystemFactories()
 				ensure(InContext.AccessVariablesStruct<FSequencePlayerData>([&AssetData, &InContext](FSequencePlayerData& InStruct)
 				{
 					InStruct.AnimSequence = AssetData.AnimationSequence;
+					InStruct.LoopMode = AssetData.LoopMode;
+					InStruct.PlayRate = AssetData.PlayRate;
+					InStruct.StartPosition = AssetData.CalculateStartTime();
 				}));
 			});
 	}
diff --git a/Engine/Plugins/Experimental/UAF/UAFAnimGraph/Source/UAFAnimGraph/Public/Asset/UAFAnimGraphAssetData.h b/Engine/Plugins/Experimental/UAF/UAFAnimGraph/Source/UAFAnimGraph/Public/Asset/UAFAnimGraphAssetData.h
--- a/Engine/Plugins/Experimental/UAF/UAFAnimGraph/Source/UAFAnimGraph/Public/Asset/UAFAnimGraphAssetData.h
+++ b/Engine/Plugins/Experimental/UAF/UAFAnimGraph/Source/UAFAnimGraph/Public/Asset/UAFAnimGraphAssetData.h
@@ -27,6 +27,13 @@ struct FUAFGraphFactoryAsset_Graph : public FUAFGraphFactoryAsset
 	TObjectPtr<const UUAFAnimGraph> AnimationGraph;
 };
 
+UENUM()
+enum class EAnimationStartTimeType : uint8
+{
+	TimeInSeconds,
+	TimePercent
+};
+
 USTRUCT(DisplayName="Animation Asset")
 struct FUAFGraphFactoryAsset_Animation : public FUAFGraphFactoryAsset
 {
@@ -38,6 +45,36 @@ struct FUAFGraphFactoryAsset_Animation : public FUAFGraphFactoryAsset
 	virtual bool Validate() const override { return AnimationSequence != nullptr; }
 
 	virtual void GetObjectReferences(TArray<const UObject*>& OutReferencedObjects) const override { OutReferencedObjects.Add(AnimationSequence); }
+
+	float CalculateStartTime() const
+	{
+		float StartTime = 0.0f;
+		float Duration = AnimationSequence != nullptr ? AnimationSequence->GetPlayLength() : 0.0f;
+		bool IsLooping = LoopMode == UE::UAF::EAnimAssetLoopMode::ForceLoop || (LoopMode == UE::UAF::EAnimAssetLoopMode::Auto && AnimationSequence != nullptr && AnimationSequence->bLoop);
+		switch (StartTimeType)
+		{
+		case EAnimationStartTimeType::TimeInSeconds:
+			StartTime = StartTimeSeconds;
+			break;
+
+		case EAnimationStartTimeType::TimePercent:
+			StartTime = Duration > UE_SMALL_NUMBER ? StartTimePercent * Duration * 0.01f : 0.0f;
+			break;
+
+		default:
+			checkNoEntry();
+		}
+
+		// Only clamp to max length if we are not looping
+		if (IsLooping)
+		{
+			return FMath::Max(StartTime, 0.0f);
+		}
+		else
+		{
+			return FMath::Clamp(StartTime, 0.0f, Duration);
+		}
+	}
 	
 	UPROPERTY(EditAnywhere, Category = AssetData)
 	TObjectPtr<const UAnimSequence> AnimationSequence;
@@ -47,6 +84,15 @@ struct FUAFGraphFactoryAsset_Animation : public FUAFGraphFactoryAsset
 
 	UPROPERTY(EditAnywhere, Category ="Playback Options")
 	UE::UAF::EAnimAssetLoopMode LoopMode = UE::UAF::EAnimAssetLoopMode::Auto;
+
+	UPROPERTY(EditAnywhere, Category ="Playback Options")
+	EAnimationStartTimeType StartTimeType = EAnimationStartTimeType::TimeInSeconds;
+	
+	UPROPERTY(EditAnywhere, Category ="Playback Options", meta=(EditConditionHides, EditCondition="StartTimeType==EAnimationStartTimeType::TimeInSeconds", Units="s", ClampMin = "0.0", UIMin = "0.0"))
+	float StartTimeSeconds = 0.0f;
+
+	UPROPERTY(EditAnywhere, Category = "Playback Options", meta=(EditConditionHides, EditCondition="StartTimeType==EAnimationStartTimeType::TimePercent", Units="Percent", ClampMin = "0.0", ClampMax = "100.0", UIMin = "0.0", UIMax = "100.0"))
+	float StartTimePercent = 0.0f;
 };
 
 USTRUCT(DisplayName="Blend Space Asset")
diff --git a/Engine/Plugins/Experimental/UAF/UAFTestSuites/Source/UAFCQTestSuite/Private/UAFAssetDataTests.cpp b/Engine/Plugins/Experimental/UAF/UAFTestSuites/Source/UAFCQTestSuite/Private/UAFAssetDataTests.cpp
--- a/Engine/Plugins/Experimental/UAF/UAFTestSuites/Source/UAFCQTestSuite/Private/UAFAssetDataTests.cpp
+++ b/Engine/Plugins/Experimental/UAF/UAFTestSuites/Source/UAFCQTestSuite/Private/UAFAssetDataTests.cpp
@@ -12,13 +12,61 @@
 #include "UAFTestsUtilities.h"
 #include "CQTest.h"
 #include "Asset/UAFAnimGraphAssetData.h"
+#include "Engine/SkeletalMesh.h"
 #include "Factory/AnimGraphFactory.h"
 #include "Traits/SequencePlayerTraitData.h"
 #include "UAF/UAFAssetData.h"
 #include "UAF/UAFAssetFactory.h"
 
+#define LOCTEXT_NAMESPACE "UAFAssetDataTests"
+
 TEST_CLASS(FUAFAssetDataTests, "Animation.UAF.Functional")
 {
+	static TObjectPtr<USkeletalMesh> GSkeletalMesh;
+	static TObjectPtr<USkeleton> GSkeleton;
+	static TObjectPtr<USkinnedAsset> GSkinnedAsset;
+	
+	BEFORE_ALL() 
+	{
+		GSkeleton = NewObject<USkeleton>();
+		{
+			FReferenceSkeletonModifier SkeletonModifier(GSkeleton);
+
+			// Add a root bone
+			SkeletonModifier.Add(FMeshBoneInfo(NAME_Root, FString(TEXT("Root")), -1), FTransform::Identity);	
+		}
+
+		GSkeletalMesh = NewObject<USkeletalMesh>();
+		GSkeletalMesh->SetSkeleton(GSkeleton);
+	}
+
+	AFTER_ALL()
+	{
+		GSkeleton = nullptr;
+		GSkeletalMesh = nullptr;
+	}
+
+	const float TestRateScale = 1.5f;
+	const int TestNumFrames = 120;
+	const FFrameRate TestFps = FFrameRate(60, 1);
+	const float TestAnimLength = TestNumFrames / TestFps.AsDecimal();
+
+	TObjectPtr<UAnimSequence> CreateAnimSequence()
+	{
+		Sequence = NewObject<UAnimSequence>();
+		Sequence->SetSkeleton(GSkeleton);
+
+		TScriptInterface<IAnimationDataController> Controller = Sequence->GetDataModel()->GetController();
+		Controller->OpenBracket(LOCTEXT("CreateAnimSequence", "Creating Animation Sequence"), false);
+		Controller->InitializeModel();
+		Controller->SetFrameRate(TestFps, false);
+		Controller->SetNumberOfFrames(TestNumFrames, false);
+		Controller->CloseBracket(false);
+		Sequence->WaitOnExistingCompression();
+
+		return Sequence;
+	}
+	
 	TObjectPtr<UUAFAssetDataTestObject> TestObject;
 	FTopLevelAssetPath UAFAssetDataTestObjectClassPath;
 	TEST_METHOD(Verify_Asset_Factory)
@@ -78,15 +126,17 @@ TEST_CLASS(FUAFAssetDataTests, "Animation.UAF.Functional")
 	
 	TObjectPtr<UAnimSequence> Sequence;
 	TInstancedStruct<FUAFGraphFactoryAsset_Animation> AnimAsset;
-	float TestRateScale = 1.5f;
+
 	UE::UAF::EAnimAssetLoopMode TestLoopMode = UE::UAF::EAnimAssetLoopMode::ForceLoop;
 	TEST_METHOD(Verify_Core_Asset_Types_Anim_Sequence)
 	{
 	
 		TestCommandBuilder.Do(TEXT("Setup"), [&]()
 			{
-				Sequence = NewObject<UAnimSequence>();
+				Sequence = CreateAnimSequence();
 				ASSERT_THAT(IsNotNull(Sequence));
+
+				ASSERT_THAT(IsNear(TestAnimLength, Sequence->GetPlayLength(), UE_SMALL_NUMBER));
 			})
 		.Then(TEXT("Generate Asset Data"), [&]()
 			{
@@ -95,6 +145,7 @@ TEST_CLASS(FUAFAssetDataTests, "Animation.UAF.Functional")
 				ASSERT_THAT(AreEqual(Sequence, AnimAsset.Get().AnimationSequence));
 				AnimAsset.GetMutable().PlayRate = TestRateScale;
 				AnimAsset.GetMutable().LoopMode = TestLoopMode;
+				
 			})
 		.Then(TEXT("Generate Factory Params From Asset"), [&]()
 			{
@@ -123,6 +174,231 @@ TEST_CLASS(FUAFAssetDataTests, "Animation.UAF.Functional")
 		});
 	}
 
+	float TestStartTimeRatio = 0.75f;
+	TEST_METHOD(Verify_Core_Asset_Types_Anim_Sequence_StartTime_Ratio)
+	{
+		TestCommandBuilder.Do(TEXT("Setup"), [&]()
+			{
+				Sequence = CreateAnimSequence();
+				ASSERT_THAT(IsNotNull(Sequence));
+
+				ASSERT_THAT(IsNear(TestAnimLength, Sequence->GetPlayLength(), UE_SMALL_NUMBER));
+			})
+		.Then(TEXT("Generate Asset Data"), [&]()
+			{
+				AnimAsset = UE::UAF::FAssetDataFactory::CreateUAFAssetDataFromObject<FUAFGraphFactoryAsset_Animation>(Sequence);
+				ASSERT_THAT(IsTrue(AnimAsset.IsValid()));
+				ASSERT_THAT(AreEqual(Sequence, AnimAsset.Get().AnimationSequence));
+				AnimAsset.GetMutable().StartTimeType = EAnimationStartTimeType::TimePercent;
+				AnimAsset.GetMutable().StartTimePercent = TestStartTimeRatio * 100.0f;
+				
+			})
+		.Then(TEXT("Generate Factory Params From Asset"), [&]()
+			{
+				FAnimNextFactoryParams Params = UE::UAF::FAnimGraphFactory::GetDefaultParamsForAsset(AnimAsset);
+				FAnimNextSimpleAnimGraphBuilder& GraphBuilder = ((FAnimNextSimpleAnimGraphBuilder&)Params.GetBuilder());
+				ASSERT_THAT(IsTrue(GraphBuilder.Stacks.Num() > 0));
+				ASSERT_THAT(IsTrue(GraphBuilder.Stacks[0].TraitDescs.Num() > 0));
+				const FAnimNextSimpleAnimGraphBuilderTraitDesc* SequencePlayerTraitData = GraphBuilder.Stacks[0].TraitDescs.FindByPredicate(
+					[](const FAnimNextSimpleAnimGraphBuilderTraitDesc& Item)
+					{
+						return Item.TraitData.GetPtr<UE::UAF::FSequencePlayerData>() != nullptr;
+					});
+				ASSERT_THAT(IsNotNull(SequencePlayerTraitData));
+
+				const UE::UAF::FSequencePlayerData* TypedSequencePlayerData = SequencePlayerTraitData->TraitData.GetPtr<UE::UAF::FSequencePlayerData>();
+				ASSERT_THAT(IsNotNull(TypedSequencePlayerData));
+				ASSERT_THAT(IsNear(TestStartTimeRatio * TestAnimLength, TypedSequencePlayerData->StartPosition, UE_SMALL_NUMBER));
+				
+			})
+		.OnTearDown(TEXT("Tear Down"), [&]()
+		{
+			AnimAsset.Reset();
+			Sequence = nullptr;
+		});
+	}
+
+	TEST_METHOD(Verify_Core_Asset_Types_Anim_Sequence_StartTime_Clamp_Time)
+	{
+	
+		TestCommandBuilder.Do(TEXT("Setup"), [&]()
+			{
+				Sequence = CreateAnimSequence();
+				ASSERT_THAT(IsNotNull(Sequence));
+
+				ASSERT_THAT(IsNear(TestAnimLength, Sequence->GetPlayLength(), UE_SMALL_NUMBER));
+			})
+		.Then(TEXT("Generate Asset Data"), [&]()
+			{
+				AnimAsset = UE::UAF::FAssetDataFactory::CreateUAFAssetDataFromObject<FUAFGraphFactoryAsset_Animation>(Sequence);
+				ASSERT_THAT(IsTrue(AnimAsset.IsValid()));
+				ASSERT_THAT(AreEqual(Sequence, AnimAsset.Get().AnimationSequence));
+				AnimAsset.GetMutable().StartTimeType = EAnimationStartTimeType::TimeInSeconds;
+				AnimAsset.GetMutable().StartTimeSeconds = 9999.f;
+				AnimAsset.GetMutable().LoopMode = UE::UAF::EAnimAssetLoopMode::ForceNonLoop;
+				
+			})
+		.Then(TEXT("Generate Factory Params From Asset"), [&]()
+			{
+				FAnimNextFactoryParams Params = UE::UAF::FAnimGraphFactory::GetDefaultParamsForAsset(AnimAsset);
+				FAnimNextSimpleAnimGraphBuilder& GraphBuilder = ((FAnimNextSimpleAnimGraphBuilder&)Params.GetBuilder());
+				ASSERT_THAT(IsTrue(GraphBuilder.Stacks.Num() > 0));
+				ASSERT_THAT(IsTrue(GraphBuilder.Stacks[0].TraitDescs.Num() > 0));
+				const FAnimNextSimpleAnimGraphBuilderTraitDesc* SequencePlayerTraitData = GraphBuilder.Stacks[0].TraitDescs.FindByPredicate(
+					[](const FAnimNextSimpleAnimGraphBuilderTraitDesc& Item)
+					{
+						return Item.TraitData.GetPtr<UE::UAF::FSequencePlayerData>() != nullptr;
+					});
+				ASSERT_THAT(IsNotNull(SequencePlayerTraitData));
+
+				const UE::UAF::FSequencePlayerData* TypedSequencePlayerData = SequencePlayerTraitData->TraitData.GetPtr<UE::UAF::FSequencePlayerData>();
+				ASSERT_THAT(IsNotNull(TypedSequencePlayerData));
+				// Should be clamped to duration
+				ASSERT_THAT(IsNear(TestAnimLength, TypedSequencePlayerData->StartPosition, UE_SMALL_NUMBER));
+				
+			})
+		.OnTearDown(TEXT("Tear Down"), [&]()
+		{
+			AnimAsset.Reset();
+			Sequence = nullptr;
+		});
+	}
+
+	TEST_METHOD(Verify_Core_Asset_Types_Anim_Sequence_StartTime_Clamp_Ratio)
+	{
+	
+		TestCommandBuilder.Do(TEXT("Setup"), [&]()
+			{
+				Sequence = CreateAnimSequence();
+				ASSERT_THAT(IsNotNull(Sequence));
+
+				ASSERT_THAT(IsNear(TestAnimLength, Sequence->GetPlayLength(), UE_SMALL_NUMBER));
+			})
+		.Then(TEXT("Generate Asset Data"), [&]()
+			{
+				AnimAsset = UE::UAF::FAssetDataFactory::CreateUAFAssetDataFromObject<FUAFGraphFactoryAsset_Animation>(Sequence);
+				ASSERT_THAT(IsTrue(AnimAsset.IsValid()));
+				ASSERT_THAT(AreEqual(Sequence, AnimAsset.Get().AnimationSequence));
+				AnimAsset.GetMutable().StartTimeType = EAnimationStartTimeType::TimePercent;
+				AnimAsset.GetMutable().StartTimePercent = 200.0f;	// Set at 200%
+				AnimAsset.GetMutable().LoopMode = UE::UAF::EAnimAssetLoopMode::ForceNonLoop;
+				
+			})
+		.Then(TEXT("Generate Factory Params From Asset"), [&]()
+			{
+				FAnimNextFactoryParams Params = UE::UAF::FAnimGraphFactory::GetDefaultParamsForAsset(AnimAsset);
+				FAnimNextSimpleAnimGraphBuilder& GraphBuilder = ((FAnimNextSimpleAnimGraphBuilder&)Params.GetBuilder());
+				ASSERT_THAT(IsTrue(GraphBuilder.Stacks.Num() > 0));
+				ASSERT_THAT(IsTrue(GraphBuilder.Stacks[0].TraitDescs.Num() > 0));
+				const FAnimNextSimpleAnimGraphBuilderTraitDesc* SequencePlayerTraitData = GraphBuilder.Stacks[0].TraitDescs.FindByPredicate(
+					[](const FAnimNextSimpleAnimGraphBuilderTraitDesc& Item)
+					{
+						return Item.TraitData.GetPtr<UE::UAF::FSequencePlayerData>() != nullptr;
+					});
+				ASSERT_THAT(IsNotNull(SequencePlayerTraitData));
+
+				const UE::UAF::FSequencePlayerData* TypedSequencePlayerData = SequencePlayerTraitData->TraitData.GetPtr<UE::UAF::FSequencePlayerData>();
+				ASSERT_THAT(IsNotNull(TypedSequencePlayerData));
+				// Should be clamped to duration
+				ASSERT_THAT(IsNear(TestAnimLength, TypedSequencePlayerData->StartPosition, UE_SMALL_NUMBER));
+				
+			})
+		.OnTearDown(TEXT("Tear Down"), [&]()
+		{
+			AnimAsset.Reset();
+			Sequence = nullptr;
+		});
+	}
+
+	float TestStartTimeSeconds = 1.15f;
+	TEST_METHOD(Verify_Core_Asset_Types_Anim_Sequence_StartTime_Seconds)
+	{
+	
+		TestCommandBuilder.Do(TEXT("Setup"), [&]()
+			{
+				Sequence = CreateAnimSequence();
+				ASSERT_THAT(IsNotNull(Sequence));
+
+				ASSERT_THAT(IsNear(TestAnimLength, Sequence->GetPlayLength(), UE_SMALL_NUMBER));
+			})
+		.Then(TEXT("Generate Asset Data"), [&]()
+			{
+				AnimAsset = UE::UAF::FAssetDataFactory::CreateUAFAssetDataFromObject<FUAFGraphFactoryAsset_Animation>(Sequence);
+				ASSERT_THAT(IsTrue(AnimAsset.IsValid()));
+				ASSERT_THAT(AreEqual(Sequence, AnimAsset.Get().AnimationSequence));
+				AnimAsset.GetMutable().StartTimeType = EAnimationStartTimeType::TimeInSeconds;
+				AnimAsset.GetMutable().StartTimeSeconds = TestStartTimeSeconds;
+				
+			})
+		.Then(TEXT("Generate Factory Params From Asset"), [&]()
+			{
+				FAnimNextFactoryParams Params = UE::UAF::FAnimGraphFactory::GetDefaultParamsForAsset(AnimAsset);
+				FAnimNextSimpleAnimGraphBuilder& GraphBuilder = ((FAnimNextSimpleAnimGraphBuilder&)Params.GetBuilder());
+				ASSERT_THAT(IsTrue(GraphBuilder.Stacks.Num() > 0));
+				ASSERT_THAT(IsTrue(GraphBuilder.Stacks[0].TraitDescs.Num() > 0));
+				const FAnimNextSimpleAnimGraphBuilderTraitDesc* SequencePlayerTraitData = GraphBuilder.Stacks[0].TraitDescs.FindByPredicate(
+					[](const FAnimNextSimpleAnimGraphBuilderTraitDesc& Item)
+					{
+						return Item.TraitData.GetPtr<UE::UAF::FSequencePlayerData>() != nullptr;
+					});
+				ASSERT_THAT(IsNotNull(SequencePlayerTraitData));
+
+				const UE::UAF::FSequencePlayerData* TypedSequencePlayerData = SequencePlayerTraitData->TraitData.GetPtr<UE::UAF::FSequencePlayerData>();
+				ASSERT_THAT(IsNotNull(TypedSequencePlayerData));
+				ASSERT_THAT(IsNear(TestStartTimeSeconds, TypedSequencePlayerData->StartPosition, UE_SMALL_NUMBER));
+				
+			})
+		.OnTearDown(TEXT("Tear Down"), [&]()
+		{
+			AnimAsset.Reset();
+			Sequence = nullptr;
+		});
+	}
+	
+	TEST_METHOD(Verify_Core_Asset_Types_Anim_Sequence_StartTime_Negative_TimeRatio_Is_Clamped)
+	{
+	
+		TestCommandBuilder.Do(TEXT("Setup"), [&]()
+			{
+				Sequence = CreateAnimSequence();
+				ASSERT_THAT(IsNotNull(Sequence));
+
+				ASSERT_THAT(IsNear(TestAnimLength, Sequence->GetPlayLength(), UE_SMALL_NUMBER));
+			})
+		.Then(TEXT("Generate Asset Data"), [&]()
+			{
+				AnimAsset = UE::UAF::FAssetDataFactory::CreateUAFAssetDataFromObject<FUAFGraphFactoryAsset_Animation>(Sequence);
+				ASSERT_THAT(IsTrue(AnimAsset.IsValid()));
+				ASSERT_THAT(AreEqual(Sequence, AnimAsset.Get().AnimationSequence));
+				AnimAsset.GetMutable().StartTimeType = EAnimationStartTimeType::TimePercent;
+				AnimAsset.GetMutable().StartTimePercent = -1.0f;
+				
+			})
+		.Then(TEXT("Generate Factory Params From Asset"), [&]()
+			{
+				FAnimNextFactoryParams Params = UE::UAF::FAnimGraphFactory::GetDefaultParamsForAsset(AnimAsset);
+				FAnimNextSimpleAnimGraphBuilder& GraphBuilder = ((FAnimNextSimpleAnimGraphBuilder&)Params.GetBuilder());
+				ASSERT_THAT(IsTrue(GraphBuilder.Stacks.Num() > 0));
+				ASSERT_THAT(IsTrue(GraphBuilder.Stacks[0].TraitDescs.Num() > 0));
+				const FAnimNextSimpleAnimGraphBuilderTraitDesc* SequencePlayerTraitData = GraphBuilder.Stacks[0].TraitDescs.FindByPredicate(
+					[](const FAnimNextSimpleAnimGraphBuilderTraitDesc& Item)
+					{
+						return Item.TraitData.GetPtr<UE::UAF::FSequencePlayerData>() != nullptr;
+					});
+				ASSERT_THAT(IsNotNull(SequencePlayerTraitData));
+
+				const UE::UAF::FSequencePlayerData* TypedSequencePlayerData = SequencePlayerTraitData->TraitData.GetPtr<UE::UAF::FSequencePlayerData>();
+				ASSERT_THAT(IsNotNull(TypedSequencePlayerData));
+				ASSERT_THAT(IsNear(0.0f, TypedSequencePlayerData->StartPosition, UE_SMALL_NUMBER));
+				
+			})
+		.OnTearDown(TEXT("Tear Down"), [&]()
+		{
+			AnimAsset.Reset();
+			Sequence = nullptr;
+		});
+	}
+
 	TObjectPtr<UBlendSpace> BlendSpace;
 	TInstancedStruct<FUAFGraphFactoryAsset_BlendSpace> BlendSpace_Asset;
 	float TestX = 1.3f;
@@ -211,5 +487,10 @@ TEST_CLASS(FUAFAssetDataTests, "Animation.UAF.Functional")
     }
 };
 
+TObjectPtr<USkeletalMesh> FUAFAssetDataTests::GSkeletalMesh = nullptr;
+TObjectPtr<USkeleton> FUAFAssetDataTests::GSkeleton = nullptr;
+
+#undef LOCTEXT_NAMESPACE
+
 #endif// WITH_EDITOR
 #endif // WITH_DEV_AUTOMATION_TESTS
\ No newline at end of file
