UAF: Input/Output rework & external anim graph ownership

Aligns 'poses' with the value runtime approach. In places where we used FAnimNextGraphLODPose, we now use FUAFValueBundle.
FUAFValueBundle now has a private implementation that can be a number of different things - a legacy FAnimNextGraphLODPose, a FValueBundleHeap, or an inline/external graph, supported by FInputValueTrait.
Adds I/O concerns to UUAFComponent - other UUAFComponents can be referenced as inputs, and output meshes can be chosen. We default to the first found skeletal mesh component if no output is set.
FAnimNextGraphInstance now has multiple lifetime stages, and can be created and owned by external systems (e.g Mass) via FAnimGraphReference (or FWeakAnimGraphReference, for non owning cases), which wrap the internal shared ptr.
This allows graph references to serve a similar purpose to FInjectionRequests. Used in verse scene graph implementations, pending in CL 49439815.

#Synced-CL 51335614

diff --git a/Engine/Plugins/Experimental/UAF/UAF/Source/UAF/Internal/Graph/RigUnit_AnimNextWriteSkeletalMeshComponentPose.h b/Engine/Plugins/Experimental/UAF/UAF/Source/UAF/Internal/Graph/RigUnit_AnimNextWriteSkeletalMeshComponentPose.h
--- a/Engine/Plugins/Experimental/UAF/UAF/Source/UAF/Internal/Graph/RigUnit_AnimNextWriteSkeletalMeshComponentPose.h
+++ b/Engine/Plugins/Experimental/UAF/UAF/Source/UAF/Internal/Graph/RigUnit_AnimNextWriteSkeletalMeshComponentPose.h
@@ -19,6 +19,8 @@ struct FRigUnit_AnimNextWriteSkeletalMeshComponentPose : public FRigUnit_AnimNex
 	RIGVM_METHOD()
 	UE_API void Execute();
 
+	static void WritePose(USkeletalMeshComponent* InOutputComponent, const FAnimNextGraphLODPose& InPose);
+
 	virtual FString GetUnitSubTitle() const { return TEXT("Skeletal Mesh Component"); };
 
 	// Pose to write
diff --git a/Engine/Plugins/Experimental/UAF/UAF/Source/UAF/Internal/Graph/RigUnit_UAFWriteOutputPose.h b/Engine/Plugins/Experimental/UAF/UAF/Source/UAF/Internal/Graph/RigUnit_UAFWriteOutputPose.h
deleted file mode 100644
diff --git a/Engine/Plugins/Experimental/UAF/UAF/Source/UAF/Internal/Graph/RigUnit_UAFWriteSystemOutput.h b/Engine/Plugins/Experimental/UAF/UAF/Source/UAF/Internal/Graph/RigUnit_UAFWriteSystemOutput.h
new file mode 100644
--- /dev/null
+++ b/Engine/Plugins/Experimental/UAF/UAF/Source/UAF/Internal/Graph/RigUnit_UAFWriteSystemOutput.h
@@ -0,0 +1,32 @@
+﻿// Copyright Epic Games, Inc. All Rights Reserved.
+
+#pragma once
+
+#include "AnimNextExecuteContext.h"
+#include "Graph/AnimNext_LODPose.h"
+#include "Graph/RigUnit_AnimNextBase.h"
+#include "UAF/ValueRuntime/ValueBundle.h"
+
+#include "RigUnit_UAFWriteSystemOutput.generated.h"
+
+#define UE_API UAF_API
+
+/** Writes the output for a system */
+USTRUCT(meta=(DisplayName="Write System Output", Category="Animation Graph", NodeColor="0, 1, 1", Keywords="Output,Pose,Port", RequiredComponents="UAFSystemOutputComponent"))
+struct FRigUnit_UAFWriteSystemOutput : public FRigUnit_AnimNextBase
+{
+	GENERATED_BODY()
+
+	RIGVM_METHOD()
+	UE_API void Execute();
+
+	// Value to write to output
+	UPROPERTY(EditAnywhere, Category = "Graph", meta = (Input))
+	FUAFValueBundle Value;
+
+	// The execution result
+	UPROPERTY(EditAnywhere, DisplayName = "Execute", Category = "BeginExecution", meta = (Input, Output))
+	FAnimNextExecuteContext ExecuteContext;
+};
+
+#undef UE_API
\ No newline at end of file
diff --git a/Engine/Plugins/Experimental/UAF/UAF/Source/UAF/Internal/Graph/UAFOutputPoseComponent.h b/Engine/Plugins/Experimental/UAF/UAF/Source/UAF/Internal/Graph/UAFOutputPoseComponent.h
deleted file mode 100644
diff --git a/Engine/Plugins/Experimental/UAF/UAF/Source/UAF/Internal/Graph/UAFSystemOutputComponent.h b/Engine/Plugins/Experimental/UAF/UAF/Source/UAF/Internal/Graph/UAFSystemOutputComponent.h
new file mode 100644
--- /dev/null
+++ b/Engine/Plugins/Experimental/UAF/UAF/Source/UAF/Internal/Graph/UAFSystemOutputComponent.h
@@ -0,0 +1,142 @@
+﻿// Copyright Epic Games, Inc. All Rights Reserved.
+
+#pragma once
+
+#include "AnimNext_LODPose.h"
+#include "RemapPoseData.h"
+#include "Module/UAFModuleInstanceComponent.h"
+
+#include "UAFSystemOutputComponent.generated.h"
+
+#define UE_API UAF_API
+
+struct FRigUnit_UAFWriteSystemOutput;
+class UMeshEntityComponentImpl_Skinned;
+class FEntitySkeletalAnimationModule;
+struct FUAFSystemOutputComponent;
+class UUAFComponent;
+struct FUAFValueBundle;
+
+namespace UE::UAF
+{
+struct FValueBundleImpl_SystemOutput;
+class ISystemOutputAdapter;
+}
+
+namespace UE::UAF::Tests
+{
+struct FUAFSystemTests;
+}
+
+// System component used to identify output values
+USTRUCT()
+struct FUAFSystemOutputComponent : public FUAFModuleInstanceComponent
+{
+	GENERATED_BODY()
+	DECLARE_UAF_ASSET_INSTANCE_COMPONENT()
+
+	// Get the refpose for the mesh we are bound to
+	const FAnimNextGraphReferencePose& GetRefPose() const
+	{
+		return RefPose;
+	}
+
+	// Get the LOD we are currently outputting as
+	int32 GetLOD() const
+	{
+		return LOD;
+	}
+
+	// Get the mesh we are currently outputting to
+	USkeletalMesh* GetSkeletalMesh() const
+	{
+		return Mesh;
+	}
+
+	// Get the legacy mesh component we are currently outputting to
+	USkeletalMeshComponent* GetSkeletalMeshComponent() const
+	{
+		return MeshComponent;
+	}
+
+	// Function used to callback into
+	using FReadOutputFunc = TFunctionRef<void(uint32 InSerialNumber, const FAnimNextGraphLODPose& InPose, const FAnimNextGraphReferencePose& InRefPose)>;
+
+	// Helper function used to remap a LOD pose to the output pose's refpose
+	// Temporary until the value runtime is ready
+	UE_API void RemapLODPoseToOutput(const FAnimNextGraphLODPose& InInputPose, FAnimNextGraphLODPose& InOutputPose);
+
+private:
+	// FUAFModuleInstanceComponent interface
+	UE_API virtual void OnBeginExecution(float InDeltaTime) override;
+
+	// Write the output
+	// @param    InValue    The value to write
+	void WriteOutput(const FUAFValueBundle& InValue);
+
+	// Generate the pose etc. for rendering by a mesh.
+	// If the serial numbers differ, a callback to the InFunction occurs, otherwise the call is skipped.
+	// If the supplied serial number matches the one held in this component, then the pose is assumed to have not been written, so no work needs to occur
+	// If the serial number supplied is zero, then the GenerateRenderData call will always be made
+	// InFunction can also be skipped if other invariants are not satisfied (e.g. null or mismatched mesh)
+	// @param    InSerialNumber    The current serial number that the reader holds
+	// @param    InFunction        Function that will be called if conditions are met
+	// @return true if the read was successful (i.e. the serial numbers were mismatched)
+	UE_API bool GenerateRenderData(uint32 InSerialNumber, FReadOutputFunc InFunction) const;
+
+	// Read the output value etc. for internal UAF use
+	// If the serial numbers differ, a callback to the InFunction occurs, otherwise the call is skipped.
+	// If the supplied serial number matches the one held in this component, then the pose is assumed to have not been written, so no work needs to occur
+	// If the serial number supplied is zero, then the GenerateRenderData call will always be made
+	// @param    InSerialNumber    The current serial number that the reader holds
+	// @param    InFunction        Function that will be called if conditions are met
+	// @return true if the read was successful (i.e. the serial numbers were mismatched)
+	UE_API bool ReadOutput(uint32 InSerialNumber, FReadOutputFunc InFunction) const;
+
+	// Called to bind this component to a specific system output adapter (e.g. mesh component)
+	UE_API void BindSystemOutputAdapter(UE::UAF::ISystemOutputAdapter& InSystemOutputAdapter);
+
+	// Unbinds from any system output adapter we are currently bound to
+	UE_API void UnbindSystemOutputAdapter();
+
+private:
+	friend FRigUnit_UAFWriteSystemOutput;
+	friend UMeshEntityComponentImpl_Skinned;
+	friend FEntitySkeletalAnimationModule;
+	friend UE::UAF::FValueBundleImpl_SystemOutput;
+	friend UUAFComponent;
+	friend UE::UAF::Tests::FUAFSystemTests;
+
+	// The output pose that was written this time around
+	FAnimNextGraphLODPose Pose;
+
+	// The mesh we use to determine the final pose format
+	UPROPERTY(Transient)
+	TObjectPtr<USkeletalMesh> Mesh;
+
+	// Legacy skeletal mesh component, if any
+	UPROPERTY(Transient)
+	TObjectPtr<USkeletalMeshComponent> MeshComponent;
+
+	// Refpose generated from the mesh
+	UPROPERTY(Transient)
+	FAnimNextGraphReferencePose RefPose;
+
+	// Remapping data used to support on-demand remaps
+	UPROPERTY(Transient)
+	FRemapPoseData RemapPoseData;
+
+	// Adapter to use for output
+	UE::UAF::ISystemOutputAdapter* SystemOutputAdapter = nullptr;
+
+	// The LOD that our system is running at this tick
+	int32 LOD = 0;
+
+	// Serial number updated each time we write the output
+	uint32 SerialNumber = 0;
+
+	// Serial number for required bones
+	uint32 CurrentRequiredBonesSerial = 0;
+};
+
+#undef UE_API
\ No newline at end of file
diff --git a/Engine/Plugins/Experimental/UAF/UAF/Source/UAF/Internal/Module/AnimNextModuleInstance.h b/Engine/Plugins/Experimental/UAF/UAF/Source/UAF/Internal/Module/AnimNextModuleInstance.h
--- a/Engine/Plugins/Experimental/UAF/UAF/Source/UAF/Internal/Module/AnimNextModuleInstance.h
+++ b/Engine/Plugins/Experimental/UAF/UAF/Source/UAF/Internal/Module/AnimNextModuleInstance.h
@@ -15,6 +15,7 @@
 #include "Module/AnimNextModuleInitMethod.h"
 #include "Module/TaskRunLocation.h"
 #include "Misc/TransactionallySafeRWLock.h"
+#include "Module/ModuleTaskContext.h"
 #include "AnimNextModuleInstance.generated.h"
 
 #define UE_API UAF_API
@@ -22,6 +23,7 @@
 struct FAnimNextModuleInstance;
 struct FRigUnit_AnimNextRunAnimationGraph_v1;
 struct FRigUnit_AnimNextRunAnimationGraph_v2;
+struct FRigUnit_UAFRunAsset;
 class UUAFEngineSubsystem;
 struct FRigUnit_AnimNextInitializeEvent;
 struct FRigUnit_CopyModuleProxyVariables;
@@ -116,6 +118,9 @@ struct FAnimNextModuleInstance : public FUAFAssetInstance
 	// Find the first 'user' tick function
 	UE_API UE::UAF::FModuleEventTickFunction* FindFirstUserTickFunction();
 
+	// Find the last 'user' tick function
+	UE_API UE::UAF::FModuleEventTickFunction* FindLastUserTickFunction();
+
 	// Run the specified script event using a module context
 	UE_API void RunScriptEvent(const UE::UAF::FScriptEvent& InEvent, float InDeltaTime);
 
@@ -131,12 +136,19 @@ struct FAnimNextModuleInstance : public FUAFAssetInstance
 	// Get a weak, non-owning reference to this instance
 	UE_API FUAFWeakSystemReference GetReference() const;
 
+	// Queue a task to run before the first user event
+	// @param	InTaskFunction		The function to run
+	void QueueTask(UE::UAF::FUniqueSystemTask&& InTaskFunction)
+	{
+		QueueTask(FName(), MoveTemp(InTaskFunction), UE::UAF::ETaskRunLocation::Before);
+	}
+
 	// Queue a task to run at a particular system event
 	// @param	InEventName			The event in the system to run the supplied task relative to. If this is NAME_None or not found then the first 'user' event will be used.
 	// @param	InTaskFunction		The function to run
 	// @param	InLocation			Where to run the task, before or after
-	UE_API void QueueTask(FName InEventName, TUniqueFunction<void(const UE::UAF::FModuleTaskContext&)>&& InTaskFunction, UE::UAF::ETaskRunLocation InLocation = UE::UAF::ETaskRunLocation::Before);
-	
+	UE_API void QueueTask(FName InEventName, UE::UAF::FUniqueSystemTask&& InTaskFunction, UE::UAF::ETaskRunLocation InLocation = UE::UAF::ETaskRunLocation::Before);
+
 	#if UAF_TRACE_ENABLED
     	// trace this module if it hasn't already been traced this frame
     	UE_API void Trace();
@@ -200,6 +212,7 @@ private:
 	friend UE::UAF::FModuleEndTickFunction;
 	friend FRigUnit_AnimNextRunAnimationGraph_v1;
 	friend FRigUnit_AnimNextRunAnimationGraph_v2;
+	friend FRigUnit_UAFRunAsset;
 	friend UUAFEngineSubsystem;
 	friend FRigUnit_AnimNextInitializeEvent;
 	friend FRigUnit_CopyModuleProxyVariables;
diff --git a/Engine/Plugins/Experimental/UAF/UAF/Source/UAF/Internal/Module/ModuleTickFunction.h b/Engine/Plugins/Experimental/UAF/UAF/Source/UAF/Internal/Module/ModuleTickFunction.h
--- a/Engine/Plugins/Experimental/UAF/UAF/Source/UAF/Internal/Module/ModuleTickFunction.h
+++ b/Engine/Plugins/Experimental/UAF/UAF/Source/UAF/Internal/Module/ModuleTickFunction.h
@@ -6,6 +6,7 @@
 #include "Engine/EngineBaseTypes.h"
 #include "Containers/SpscQueue.h"
 #include "TaskSyncManager.h"
+#include "Module/ModuleTaskContext.h"
 #include "Script/UAFScriptComponent.h"
 #if UE_AUTORTFM
 #include "Containers/TransactionallySafeSpscQueue.h"
@@ -61,9 +62,9 @@ struct FModuleEventTickFunction : public FTickFunction
 
 	// Tasks queued to execute once per event
 #if UE_AUTORTFM
-	using FTaskQueueType = TTransactionallySafeSpscQueue<TUniqueFunction<void(const FModuleTaskContext&)>>;
+	using FTaskQueueType = TTransactionallySafeSpscQueue<UE::UAF::FUniqueSystemTask>;
 #else
-	using FTaskQueueType = TMpscQueue<TUniqueFunction<void(const FModuleTaskContext&)>>;
+	using FTaskQueueType = TMpscQueue<UE::UAF::FUniqueSystemTask>;
 #endif
 
 	FTaskQueueType PreExecuteTasks;
diff --git a/Engine/Plugins/Experimental/UAF/UAF/Source/UAF/Internal/UAF/ValueRuntime/IVirtualValueBundle.h b/Engine/Plugins/Experimental/UAF/UAF/Source/UAF/Internal/UAF/ValueRuntime/IVirtualValueBundle.h
new file mode 100644
--- /dev/null
+++ b/Engine/Plugins/Experimental/UAF/UAF/Source/UAF/Internal/UAF/ValueRuntime/IVirtualValueBundle.h
@@ -0,0 +1,60 @@
+﻿// Copyright Epic Games, Inc. All Rights Reserved.
+
+#pragma once
+
+#include "CoreTypes.h"
+#include "Containers/Array.h"
+#include "Misc/MemStack.h"
+
+#define UE_API UAF_API
+
+struct FAnimNextGraphInstance;
+struct FRigUnit_ReadInputPose;
+struct FAnimNextGraphLODPose;
+
+namespace UE::UAF
+{
+class FValueBundleHeap;
+struct FInputValueTrait;
+struct FEvaluateTraversalContext;
+struct FUpdateTraversalContext;
+template<class TraitInterfaceType>
+struct TTraitBinding;
+struct IEvaluate;
+struct IUpdate;
+struct IUpdateTraversal;
+struct IHierarchy;
+struct FWeakTraitPtr;
+using FChildrenArray = TArray<FWeakTraitPtr, TInlineAllocator<8, TMemStackAllocator<>>>;;
+struct FTraitUpdateState;
+struct FUpdateTraversalQueue;
+struct FExecutionContext;
+struct FWeakAnimGraphReference;
+}
+
+namespace UE::UAF
+{
+
+// Implementation of a value bundle 'virtual value' that can represent a number of underlying concepts, e.g.:
+// - A concrete value bundle
+// - A legacy LOD pose
+// - The result of reading a referenced system's output
+// - An inline or owned graph instance
+class IVirtualValueBundle
+{
+public:
+	virtual ~IVirtualValueBundle() = default;
+
+	// Get the LOD pose this bundle represents
+	virtual const FAnimNextGraphLODPose* GetLODPose() const = 0;
+
+	// Get a value bundle - valid if this bundle is a concrete value
+	virtual const FValueBundleHeap* GetValueBundle() const = 0;
+
+	// Get an anim graph reference that represents a 'virtual' value bundle - i.e. the result of evaluating the graph in-line
+	virtual const FWeakAnimGraphReference* GetAnimGraphReference() const = 0;
+};
+
+}
+
+#undef UE_API
\ No newline at end of file
diff --git a/Engine/Plugins/Experimental/UAF/UAF/Source/UAF/Private/AnimNextModuleImpl.cpp b/Engine/Plugins/Experimental/UAF/UAF/Source/UAF/Private/AnimNextModuleImpl.cpp
--- a/Engine/Plugins/Experimental/UAF/UAF/Source/UAF/Private/AnimNextModuleImpl.cpp
+++ b/Engine/Plugins/Experimental/UAF/UAF/Source/UAF/Private/AnimNextModuleImpl.cpp
@@ -37,6 +37,7 @@
 #include "GameFramework/CharacterMovementComponent.h"
 #include "Module/UAFSystemAssetData.h"
 #include "UAF/UAFAssetFactory.h"
+#include "UAF/ValueRuntime/ValueBundle.h"
 #include "UAF/Attributes/EngineAttributes.h"
 
 #define LOCTEXT_NAMESPACE "AnimNextModule"
@@ -86,6 +87,7 @@ namespace UE::UAF
 			FUAFWeakSystemReference::StaticStruct(),
 			FAnimNextGraphLODPose::StaticStruct(),
 			FAnimNextGraphReferencePose::StaticStruct(),
+			FUAFValueBundle::StaticStruct(),
 		};
 
 		RigVMRegistry.RegisterStructTypes(AllowedStructTypes);
diff --git a/Engine/Plugins/Experimental/UAF/UAF/Source/UAF/Private/Component/AnimNextComponent.cpp b/Engine/Plugins/Experimental/UAF/UAF/Source/UAF/Private/Component/AnimNextComponent.cpp
--- a/Engine/Plugins/Experimental/UAF/UAF/Source/UAF/Private/Component/AnimNextComponent.cpp
+++ b/Engine/Plugins/Experimental/UAF/UAF/Source/UAF/Private/Component/AnimNextComponent.cpp
@@ -4,6 +4,8 @@
 
 #include "Blueprint/BlueprintExceptionInfo.h"
 #include "Engine/World.h"
+#include "Graph/RigUnit_AnimNextWriteSkeletalMeshComponentPose.h"
+#include "Graph/UAFSystemOutputComponent.h"
 #include "Module/ModuleTaskContext.h"
 #include "UObject/FortniteMainBranchObjectVersion.h"
 #include "Variables/AnimNextVariableReference.h"
@@ -12,6 +14,8 @@
 #include "Module/UAFSystemAssetData.h"
 #include "Module/UAFWeakSystemReference.h"
 #include "UAF/UAFAssetFactory.h"
+#include "UAF/ValueRuntime/ValueBundle.h"
+#include "UAF/ValueRuntime/ValueBundleImpl_SystemOutput.h"
 
 #include UE_INLINE_GENERATED_CPP_BY_NAME(AnimNextComponent)
 
@@ -23,9 +27,9 @@ void UUAFComponent::OnRegister()
 
 	if (AssetData.IsValid() && AssetData.Get().Validate())
 	{
-		check(!SystemReference.IsValid());
+		AllocateSystem();
 
-		SystemReference = FSystemReference(AssetData, this, InitMethod);
+		SetupInputOutput();
 
 #if WITH_EDITOR
 		SystemReference.RunInitialTickInEditor();
@@ -37,6 +41,16 @@ void UUAFComponent::OnRegister()
 	}
 }
 
+void UUAFComponent::AllocateSystem()
+{
+	using namespace UE::UAF;
+
+	if(AssetData.IsValid() && AssetData.Get().Validate() && !SystemReference.IsValid())
+	{
+		SystemReference = FSystemReference(AssetData, this, InitMethod);
+	}
+}
+
 void UUAFComponent::OnUnregister()
 {
 	Super::OnUnregister();
@@ -96,6 +110,28 @@ void UUAFComponent::Serialize(FArchive& Ar)
 #endif
 }
 
+UE::UAF::ISystemOutputAdapter::FInputData UUAFComponent::GetInputData() const
+{
+	return UE::UAF::ISystemOutputAdapter::FInputData(CachedSkeletalMeshComponent.Get());
+}
+
+void UUAFComponent::SignalOutputWritten(const UE::UAF::FModuleTaskContext& InContext)
+{
+	if (SkeletalMeshComponentOutputMode == EUAFSkeletalMeshComponentOutputMode::WriteToSkeletalMeshComponentPose)
+	{
+		InContext.AccessComponent<FUAFSystemOutputComponent>([this](FUAFSystemOutputComponent& InOutputPoseComponent)
+		{
+			InOutputPoseComponent.GenerateRenderData(SerialNumber, [this](uint32 InSerialNumber, const FAnimNextGraphLODPose& InPose, const FAnimNextGraphReferencePose& InRefPose)
+			{
+				SerialNumber = InSerialNumber;
+
+				// TODO: for now, hook the legacy output pose path. this needs a refactor of skel mesh component to work better
+				FRigUnit_AnimNextWriteSkeletalMeshComponentPose::WritePose(CachedSkeletalMeshComponent.Get(), InPose);
+			});
+		});
+	}
+}
+
 void UUAFComponent::BlueprintSetVariable(FName Name, int32 Value)
 {
 	checkNoEntry();
@@ -518,3 +554,104 @@ bool UUAFComponent::IsModuleValid()
 	return SystemReference.IsValid();
 }
 
+void UUAFComponent::SetupInputOutput()
+{
+	check(IsInGameThread());
+
+	CachedSkeletalMeshComponent = DetermineSkeletalMeshComponent();
+	const bool bValidOutputComponent = (CachedSkeletalMeshComponent != nullptr);
+	if (!bValidOutputComponent)
+	{
+		UE_LOGFMT(LogAnimation, Warning, "UUAFComponent::SetupInputOutput: Could not find a valid mesh component to use for system output");
+	}
+
+	struct FInputBinding
+	{
+		FAnimNextVariableReference Variable;
+		TWeakObjectPtr<UUAFComponent> Component;
+	};
+
+	AActor* Owner = GetOwner();
+	check(Owner);
+	TArray<FInputBinding> InputBindings;
+	InputBindings.Reserve(Inputs.Num());
+	for (const FUAFComponentInputDesc& InputDesc : Inputs)
+	{
+		if (InputDesc.Input.IsValid())
+		{
+			if (UUAFComponent* InputComponent = Cast<UUAFComponent>(InputDesc.Component.GetComponent(Owner)))
+			{
+				if (InputComponent == this)
+				{
+					UE_LOGFMT(LogAnimation, Warning, "UUAFComponent::SetupInputOutput: Cannot apply self as input pose");
+					continue;
+				}
+
+				InputBindings.Add({ .Variable = InputDesc.Input, .Component = InputComponent });
+			}
+			else
+			{
+				UE_LOGFMT(LogAnimation, Warning, "UUAFComponent::SetupInputOutput: Invalid input pose component specified");
+			}
+		}
+		else
+		{
+			UE_LOGFMT(LogAnimation, Warning, "UUAFComponent::SetupInputOutput: Invalid input pose variable specified");
+		}
+	}
+
+	// Set up tick dependencies for valid inputs
+	// By default we link [Input Last Event] -> [Our First Event].
+	// TODO: Allow user-specified dependencies here
+	for (const FInputBinding& InputBinding : InputBindings)
+	{
+		// Allocate the prerequisite system in case it has not been registered yet
+		InputBinding.Component->AllocateSystem();
+
+		UE::UAF::FSystemReference& InputSystem = InputBinding.Component->SystemReference;
+		SystemReference.AddSystemEventDependency(SystemReference.GetFirstUserEventName(), FUAFWeakSystemReference(InputSystem), InputSystem.GetLastUserEventName(), UE::UAF::ESystemDependency::Prerequisite);
+	}
+
+	SystemReference.QueueTask([WeakThis = TWeakObjectPtr<UUAFComponent>(this), bValidOutputComponent, InputBindings = MoveTemp(InputBindings)](const UE::UAF::FModuleTaskContext& InContext)
+	{
+		if (UUAFComponent* This = WeakThis.Get())
+		{
+			if (bValidOutputComponent)
+			{
+				InContext.AccessComponent<FUAFSystemOutputComponent>([This](FUAFSystemOutputComponent& UAFSystemOutputComponent)
+				{
+					UAFSystemOutputComponent.BindSystemOutputAdapter(*This);
+				});
+			}
+
+			for (const FInputBinding& InputBinding : InputBindings)
+			{
+				if (UUAFComponent* Component = InputBinding.Component.Get())
+				{
+					if (!InContext.AccessVariable<FUAFValueBundle>(InputBinding.Variable, [Component](FUAFValueBundle& InValueBundle)
+					{
+						InValueBundle.SetAs<UE::UAF::FValueBundleImpl_SystemOutput>(FUAFWeakSystemReference(Component->SystemReference));
+					}))
+					{
+						UE_LOGFMT(LogAnimation, Warning, "UUAFComponent::SetupInputOutput: Cannot bind input pose variable: '{0}'", InputBinding.Variable.GetName());
+					}
+				}
+			}
+		}
+	});
+}
+
+USkeletalMeshComponent* UUAFComponent::DetermineSkeletalMeshComponent() const
+{
+	AActor* Owner = GetOwner();
+	if (Owner == nullptr)
+	{
+		return nullptr;
+	}
+	
+	if (USkeletalMeshComponent* UserSpecifiedComponent = Cast<USkeletalMeshComponent>(OutputComponent.GetComponent(Owner)))
+	{
+		return UserSpecifiedComponent;
+	}
+	return Owner->FindComponentByClass<USkeletalMeshComponent>();
+}
diff --git a/Engine/Plugins/Experimental/UAF/UAF/Source/UAF/Private/Graph/RigUnit_AnimNextWriteSkeletalMeshComponentPose.cpp b/Engine/Plugins/Experimental/UAF/UAF/Source/UAF/Private/Graph/RigUnit_AnimNextWriteSkeletalMeshComponentPose.cpp
--- a/Engine/Plugins/Experimental/UAF/UAF/Source/UAF/Private/Graph/RigUnit_AnimNextWriteSkeletalMeshComponentPose.cpp
+++ b/Engine/Plugins/Experimental/UAF/UAF/Source/UAF/Private/Graph/RigUnit_AnimNextWriteSkeletalMeshComponentPose.cpp
@@ -20,11 +20,6 @@ DEFINE_STAT(STAT_AnimNext_Write_Pose);
 
 FRigUnit_AnimNextWriteSkeletalMeshComponentPose_Execute()
 {
-	SCOPE_CYCLE_COUNTER(STAT_AnimNext_Write_Pose);
-
-	using namespace UE::UAF;
-	FTaskTagScope TaskTag(ETaskTag::EParallelGameThread);
-
 	USkeletalMeshComponent* OutputComponent = SkeletalMeshComponent;
 
 	// Defer to module component if no component is supplied
@@ -36,78 +31,88 @@ FRigUnit_AnimNextWriteSkeletalMeshComponentPose_Execute()
 		OutputComponent = ComponentReference.GetComponent();
 	}
 
-	if(OutputComponent == nullptr)
+	WritePose(OutputComponent, Pose);
+}
+
+void FRigUnit_AnimNextWriteSkeletalMeshComponentPose::WritePose(USkeletalMeshComponent* InOutputComponent, const FAnimNextGraphLODPose& InPose)
+{
+	SCOPE_CYCLE_COUNTER(STAT_AnimNext_Write_Pose);
+
+	using namespace UE::UAF;
+	FTaskTagScope TaskTag(ETaskTag::EParallelGameThread);
+
+	if(InOutputComponent == nullptr)
 	{
-		UAF_RIGUNIT_LOG(Warning, TEXT("Could not write to skeletal mesh component - No valid skeletal mesh component found."));
+		UE_LOG(LogAnimation, Warning, TEXT("Could not write to skeletal mesh component - No valid skeletal mesh component found."));
 		return;
 	}
 
-	if(!Pose.LODPose.IsValid())
+	if(!InPose.LODPose.IsValid())
 	{
-		UAF_RIGUNIT_LOG(Warning, TEXT("Could not write to skeletal mesh component - LOD Pose is not valid."));
+		UE_LOG(LogAnimation, Warning, TEXT("Could not write to skeletal mesh component - LOD Pose is not valid."));
 		return;
 	}
-
-	USkeletalMesh* SkeletalMesh = OutputComponent->GetSkeletalMeshAsset();
+	
+	USkeletalMesh* SkeletalMesh = InOutputComponent->GetSkeletalMeshAsset();
 	if(SkeletalMesh == nullptr)
 	{
-		UAF_RIGUNIT_LOG(Warning, TEXT("Could not write to skeletal mesh component - No valid SkeletalMesh on [SK: %s][Owner: %s]"),
-			*OutputComponent->GetName(), *GetNameSafe(OutputComponent->GetOuter()));
+		UE_LOG(LogAnimation, Warning, TEXT("Could not write to skeletal mesh component - No valid SkeletalMesh on [SK: %s][Owner: %s]"),
+			*InOutputComponent->GetName(), *GetNameSafe(InOutputComponent->GetOuter()));
 		return;
 	}
 
 	// We cannot write to the skeletal mesh component if it is using anim BP
-	if(OutputComponent->bEnableAnimation)
+	if(InOutputComponent->bEnableAnimation)
 	{
-		UAF_RIGUNIT_LOG(Warning, TEXT("Could not write to skeletal mesh component - bEnableAnimation is true[SK: %s][Owner: %s]"),
-			*OutputComponent->GetName(), *GetNameSafe(OutputComponent->GetOuter()));
+		UE_LOG(LogAnimation, Warning, TEXT("Could not write to skeletal mesh component - bEnableAnimation is true[SK: %s][Owner: %s]"),
+			*InOutputComponent->GetName(), *GetNameSafe(InOutputComponent->GetOuter()));
 		return;
 	}
 
-	FDataHandle RefPoseHandle = FDataRegistry::Get()->GetOrGenerateReferencePose(OutputComponent);
+	FDataHandle RefPoseHandle = FDataRegistry::Get()->GetOrGenerateReferencePose(InOutputComponent);
 	const UE::UAF::FReferencePose& RefPose = RefPoseHandle.GetRef<UE::UAF::FReferencePose>();
 
 	// We cannot use two different reference poses because we want to avoid an expensive remap operation
 	// You should remap the pose explicitly if this is what you need
-	if (&RefPose != &Pose.LODPose.GetRefPose())
+	if (&RefPose != &InPose.LODPose.GetRefPose())
 	{
-		UAF_RIGUNIT_LOG(Warning, TEXT("Could not write to skeletal mesh component - The input pose and the skeletal mesh component use different reference poses [SK: %s][Owner: %s]"),
-			*SkeletalMesh->GetName(), *GetNameSafe(OutputComponent->GetOuter()));
+		UE_LOG(LogAnimation, Warning, TEXT("Could not write to skeletal mesh component - The input pose and the skeletal mesh component use different reference poses [SK: %s][Owner: %s]"),
+			*SkeletalMesh->GetName(), *GetNameSafe(InOutputComponent->GetOuter()));
 
 		return;
 	}
 
 	FMemMark MemMark(FMemStack::Get());
-	USkinnedMeshComponent::FRenderStateLockScope LockScope(OutputComponent);
+	USkinnedMeshComponent::FRenderStateLockScope LockScope(InOutputComponent);
 
 	FBlendedHeapCurve PrevCurves;
 
 	// Clear our curves and attributes or we'll have leftovers from our last write.
 	PRAGMA_DISABLE_DEPRECATION_WARNINGS
-	Swap(PrevCurves, OutputComponent->AnimCurves);
+	Swap(PrevCurves, InOutputComponent->AnimCurves);
 	PRAGMA_ENABLE_DEPRECATION_WARNINGS
-	OutputComponent->GetEditableCustomAttributes().Empty();
+	InOutputComponent->GetEditableCustomAttributes().Empty();
 
 	PRAGMA_DISABLE_DEPRECATION_WARNINGS
-	OutputComponent->AnimCurves.CopyFrom(Pose.Curves);
+	InOutputComponent->AnimCurves.CopyFrom(InPose.Curves);
 	PRAGMA_ENABLE_DEPRECATION_WARNINGS
 
 #if WITH_EDITOR
 	PRAGMA_DISABLE_DEPRECATION_WARNINGS
-	OutputComponent->GetEditableAnimationCurves().CopyFrom(Pose.Curves);
+	InOutputComponent->GetEditableAnimationCurves().CopyFrom(InPose.Curves);
 	PRAGMA_ENABLE_DEPRECATION_WARNINGS
 #endif
 
 	TUniqueFunction<void(USkinnedMeshComponent*)> GameThreadCallback;
 
 	// Extract our material parameters and morph target curves
-	if (Pose.Curves.Num() != 0 && RefPose.CurveFlags.Num() != 0)
+	if (InPose.Curves.Num() != 0 && RefPose.CurveFlags.Num() != 0)
 	{
 		FBlendedCurve ExpiredCurves;
 		ExpiredCurves.Reserve(64);			// Expired curves are uncommon
 
 		// Reset any material curves we produced in the last update but not in this update
-		UE::Anim::FNamedValueArrayUtils::Subtraction(PrevCurves, Pose.Curves,
+		UE::Anim::FNamedValueArrayUtils::Subtraction(PrevCurves, InPose.Curves,
 			[&ExpiredCurves](const UE::Anim::FCurveElement& InCurveElement)
 			{
 				ExpiredCurves.Add(InCurveElement);
@@ -131,7 +136,7 @@ FRigUnit_AnimNextWriteSkeletalMeshComponentPose_Execute()
 				}
 			});
 
-		UE::Anim::FNamedValueArrayUtils::Intersection(Pose.Curves, RefPose.CurveFlags,
+		UE::Anim::FNamedValueArrayUtils::Intersection(InPose.Curves, RefPose.CurveFlags,
 			[&MaterialCurves, &MorphTargetCurves](const UE::Anim::FCurveElement& InCurveElement, const UE::Anim::FCurveElementFlags& InCurveFlagsElement)
 			{
 				if (EnumHasAnyFlags(InCurveFlagsElement.Flags | InCurveElement.Flags, UE::Anim::ECurveElementFlags::MorphTarget))
@@ -198,11 +203,11 @@ FRigUnit_AnimNextWriteSkeletalMeshComponentPose_Execute()
 	}
 
 	// Attributes require remapping since the indices are LOD indices and we want mesh indices
-	FGenerationTools::RemapAttributes(Pose.LODPose, Pose.Attributes, OutputComponent->GetEditableCustomAttributes());
+	FGenerationTools::RemapAttributes(InPose.LODPose, InPose.Attributes, InOutputComponent->GetEditableCustomAttributes());
 
 	// Convert and dispatch to renderer
 	UE::Anim::FSkinnedMeshComponentExtensions::CompleteAndDispatch(
-		OutputComponent,
-		Pose.LODPose,
+		InOutputComponent,
+		InPose.LODPose,
 		MoveTemp(GameThreadCallback));
-}
+}
\ No newline at end of file
diff --git a/Engine/Plugins/Experimental/UAF/UAF/Source/UAF/Private/Graph/RigUnit_UAFWriteOutputPose.cpp b/Engine/Plugins/Experimental/UAF/UAF/Source/UAF/Private/Graph/RigUnit_UAFWriteOutputPose.cpp
deleted file mode 100644
diff --git a/Engine/Plugins/Experimental/UAF/UAF/Source/UAF/Private/Graph/RigUnit_UAFWriteSystemOutput.cpp b/Engine/Plugins/Experimental/UAF/UAF/Source/UAF/Private/Graph/RigUnit_UAFWriteSystemOutput.cpp
new file mode 100644
--- /dev/null
+++ b/Engine/Plugins/Experimental/UAF/UAF/Source/UAF/Private/Graph/RigUnit_UAFWriteSystemOutput.cpp
@@ -0,0 +1,21 @@
+﻿// Copyright Epic Games, Inc. All Rights Reserved.
+
+#include "Graph/RigUnit_UAFWriteSystemOutput.h"
+
+#include "AnimNextStats.h"
+#include "Graph/UAFSystemOutputComponent.h"
+#include "Module/AnimNextModuleInstance.h"
+
+#include UE_INLINE_GENERATED_CPP_BY_NAME(RigUnit_UAFWriteSystemOutput)
+
+FRigUnit_UAFWriteSystemOutput_Execute()
+{
+	SCOPE_CYCLE_COUNTER(STAT_AnimNext_Write_Pose);
+
+	using namespace UE::UAF;
+
+	const FAnimNextModuleContextData& ModuleContextData = ExecuteContext.GetContextData<FAnimNextModuleContextData>();
+	FAnimNextModuleInstance& ModuleInstance = ModuleContextData.GetModuleInstance();
+	FUAFSystemOutputComponent& OutputPoseComponent = ModuleInstance.GetComponent<FUAFSystemOutputComponent>();
+	OutputPoseComponent.WriteOutput(Value);
+}
diff --git a/Engine/Plugins/Experimental/UAF/UAF/Source/UAF/Private/Graph/UAFOutputPoseComponent.cpp b/Engine/Plugins/Experimental/UAF/UAF/Source/UAF/Private/Graph/UAFOutputPoseComponent.cpp
deleted file mode 100644
diff --git a/Engine/Plugins/Experimental/UAF/UAF/Source/UAF/Private/Graph/UAFSystemOutputComponent.cpp b/Engine/Plugins/Experimental/UAF/UAF/Source/UAF/Private/Graph/UAFSystemOutputComponent.cpp
new file mode 100644
--- /dev/null
+++ b/Engine/Plugins/Experimental/UAF/UAF/Source/UAF/Private/Graph/UAFSystemOutputComponent.cpp
@@ -0,0 +1,171 @@
+﻿// Copyright Epic Games, Inc. All Rights Reserved.
+
+#include "Graph/UAFSystemOutputComponent.h"
+
+#include "AnimationRuntime.h"
+#include "ISystemOutputAdapter.h"
+#include "ReferencePose.h"
+#include "Engine/SkeletalMesh.h"
+#include "Graph/AnimNext_LODPose.h"
+#include "Module/AnimNextModuleInstance.h"
+#include "Module/ModuleTaskContext.h"
+#include "UAF/ValueRuntime/IVirtualValueBundle.h"
+#include "UAF/ValueRuntime/ValueBundle.h"
+
+void FUAFSystemOutputComponent::OnBeginExecution(float InDeltaTime)
+{
+	// Grab all the info we need to run the system this time around & generate/get the reference pose
+	if (SystemOutputAdapter != nullptr)
+	{
+		UE::UAF::ISystemOutputAdapter::FInputData InputData = SystemOutputAdapter->GetInputData();
+		USkeletalMesh* OldMesh = Mesh;
+		Mesh = InputData.SkeletalMesh;
+		MeshComponent = InputData.SkeletalMeshComponent;
+		LOD = InputData.LOD;
+
+		if (MeshComponent)
+		{
+			if (MeshComponent->GetSkeletalMeshAsset() != OldMesh)
+			{
+				RefPose.ReferencePose = UE::UAF::FDataRegistry::Get()->GetOrGenerateReferencePose(MeshComponent);
+			}
+		}
+		else if (Mesh)
+		{
+			if (Mesh != OldMesh)
+			{
+				RefPose.ReferencePose = UE::UAF::FDataRegistry::Get()->GetOrGenerateReferencePose(Mesh);
+			}
+		}
+		else
+		{
+			RefPose.ReferencePose = UE::UAF::FDataHandle();
+		}
+	}
+}
+
+void FUAFSystemOutputComponent::WriteOutput(const FUAFValueBundle& InValue)
+{
+	SerialNumber++;
+
+	// If serial number wrapped, skip zero so client serials that have just been created will always mismatch
+	if (SerialNumber == 0)
+	{
+		SerialNumber++;
+	}
+	
+	if (InValue.IsRefPose())
+	{
+		Pose.LODPose.SetRefPose();
+	}
+	else
+	{
+		UE::UAF::IVirtualValueBundle& Impl = InValue.GetImplChecked();
+
+		// TODO: Value bundle version of this
+		if (const FAnimNextGraphLODPose* ValuePose = Impl.GetLODPose())
+		{
+			RemapLODPoseToOutput(*ValuePose, Pose);
+		}
+	}
+
+	if (SystemOutputAdapter != nullptr)
+	{
+		SystemOutputAdapter->SignalOutputWritten(UE::UAF::FModuleTaskContext(GetModuleInstance()));
+	}
+}
+
+bool FUAFSystemOutputComponent::GenerateRenderData(uint32 InSerialNumber, FReadOutputFunc InFunction) const
+{
+	// Zero serial indicates no valid pose to output, so early out
+	if (SerialNumber == 0)
+	{
+		return false;
+	}
+
+	// Always call the read function when the supplied serial is zero, otherwise only call it on mismatched serial numbers
+	if (InSerialNumber != 0 && InSerialNumber == SerialNumber)
+	{
+		return false;
+	}
+
+	if (!ensure(Mesh != nullptr))
+	{
+		return false;
+	}
+
+	if (SystemOutputAdapter == nullptr)
+	{
+		return false;
+	}
+
+	// Validate that the input we are transmitting matches the expected mesh of the adapter
+	// It is the assumed responsibility of client systems to ensure consistency of the mesh across the frame.
+	// It may be necessary when switching mesh mid-frame to recreate a new system to perform the work again.
+	// Note LOD can vary across the frame due to streaming and unfenced other LOD calculations, so clients must deal with LODs that are
+	// inconsistent between OnBeginExecution's GetInputData and this call
+	UE::UAF::ISystemOutputAdapter::FInputData InputData = SystemOutputAdapter->GetInputData();
+	if (ensure(InputData.SkeletalMesh == Mesh))
+	{
+		InFunction(SerialNumber, Pose, RefPose);
+		return true;
+	}
+
+	return false;
+}
+
+bool FUAFSystemOutputComponent::ReadOutput(uint32 InSerialNumber, FReadOutputFunc InFunction) const
+{
+	// Zero serial indicates no valid pose to output, so early out
+	if (SerialNumber == 0)
+	{
+		return false;
+	}
+
+	// Always call the read function when the supplied serial is zero, otherwise only call it on mismatched serial numbers
+	if (InSerialNumber != 0 && InSerialNumber == SerialNumber)
+	{
+		return false;
+	}
+
+	InFunction(SerialNumber, Pose, RefPose);
+	return true;
+}
+
+void FUAFSystemOutputComponent::BindSystemOutputAdapter(UE::UAF::ISystemOutputAdapter& InSystemOutputAdapter)
+{
+	ensure(SystemOutputAdapter == nullptr);	// Should unbind first
+	SystemOutputAdapter = &InSystemOutputAdapter;
+}
+
+void FUAFSystemOutputComponent::UnbindSystemOutputAdapter()
+{
+	ensure(SystemOutputAdapter != nullptr);	// Should bind first
+	SystemOutputAdapter = nullptr;
+}
+
+void FUAFSystemOutputComponent::RemapLODPoseToOutput(const FAnimNextGraphLODPose& InInputPose, FAnimNextGraphLODPose& InOutputPose)
+{
+	const UE::UAF::FReferencePose* InputRefPose = InInputPose.LODPose.RefPose;
+	const UE::UAF::FReferencePose* OutputRefPose = RefPose.ReferencePose.IsValid() ? RefPose.ReferencePose.GetPtr<UE::UAF::FReferencePose>() : nullptr;
+	if (InputRefPose != nullptr && OutputRefPose != nullptr)
+	{
+		if (InputRefPose != OutputRefPose)
+		{
+			// Different reference pose, so we must lazily remap here
+			// Pre-create the mapping and cache it to avoid runtime lookups.
+			if (RemapPoseData.ShouldReinit(*InputRefPose, *OutputRefPose))
+			{
+				RemapPoseData.Reinit(*InputRefPose, *OutputRefPose);
+			}
+
+			InOutputPose.Curves.CopyFrom(InInputPose.Curves);
+			RemapPoseData.RemapPose(InInputPose.LODPose, InOutputPose.LODPose);
+			RemapPoseData.RemapAttributes(InInputPose.LODPose, InInputPose.Attributes, InOutputPose.LODPose, InOutputPose.Attributes);
+		}
+		else
+		{
+			InOutputPose = InInputPose;
+		}
+	}
+}
\ No newline at end of file
diff --git a/Engine/Plugins/Experimental/UAF/UAF/Source/UAF/Private/Module/AnimNextModuleInstance.cpp b/Engine/Plugins/Experimental/UAF/UAF/Source/UAF/Private/Module/AnimNextModuleInstance.cpp
--- a/Engine/Plugins/Experimental/UAF/UAF/Source/UAF/Private/Module/AnimNextModuleInstance.cpp
+++ b/Engine/Plugins/Experimental/UAF/UAF/Source/UAF/Private/Module/AnimNextModuleInstance.cpp
@@ -632,7 +632,17 @@ UE::UAF::FModuleEventTickFunction* FAnimNextModuleInstance::FindFirstUserTickFun
 	return FoundTickFunction;
 }
 
-void FAnimNextModuleInstance::QueueTask(FName InEventName, TUniqueFunction<void(const UE::UAF::FModuleTaskContext&)>&& InTaskFunction, UE::UAF::ETaskRunLocation InLocation)
+UE::UAF::FModuleEventTickFunction* FAnimNextModuleInstance::FindLastUserTickFunction()
+{
+	const int32 LastUserIndex = TickFunctions.FindLastByPredicate([](const UE::UAF::FModuleEventTickFunction& InTickFunction)
+	{
+		return InTickFunction.bUserEvent;
+	});
+
+	return LastUserIndex != INDEX_NONE ? &TickFunctions[LastUserIndex] : nullptr;
+}
+
+void FAnimNextModuleInstance::QueueTask(FName InEventName, UE::UAF::FUniqueSystemTask&& InTaskFunction, UE::UAF::ETaskRunLocation InLocation)
 {
 	using namespace UE::UAF;
 
diff --git a/Engine/Plugins/Experimental/UAF/UAF/Source/UAF/Private/Module/ModuleTaskContext.cpp b/Engine/Plugins/Experimental/UAF/UAF/Source/UAF/Private/Module/ModuleTaskContext.cpp
--- a/Engine/Plugins/Experimental/UAF/UAF/Source/UAF/Private/Module/ModuleTaskContext.cpp
+++ b/Engine/Plugins/Experimental/UAF/UAF/Source/UAF/Private/Module/ModuleTaskContext.cpp
@@ -34,6 +34,12 @@ void FModuleTaskContext::TryAccessComponent(UScriptStruct* InComponentType, TFun
 	InFunction(*Component);
 }
 
+void FModuleTaskContext::AccessComponent(UScriptStruct* InComponentType, TFunctionRef<void(FUAFAssetInstanceComponent&)> InFunction) const
+{
+	FUAFAssetInstanceComponent& Component = ModuleInstance.GetComponent(InComponentType);
+	InFunction(Component);
+}
+
 FAnimNextModuleInstance* const FModuleTaskContext::GetModuleInstance() const
 {
 	return &ModuleInstance;
diff --git a/Engine/Plugins/Experimental/UAF/UAF/Source/UAF/Private/Module/SystemReference.cpp b/Engine/Plugins/Experimental/UAF/UAF/Source/UAF/Private/Module/SystemReference.cpp
--- a/Engine/Plugins/Experimental/UAF/UAF/Source/UAF/Private/Module/SystemReference.cpp
+++ b/Engine/Plugins/Experimental/UAF/UAF/Source/UAF/Private/Module/SystemReference.cpp
@@ -304,6 +304,32 @@ const FTickFunction* FSystemReference::FindTickFunction(FName InEventName) const
 	return nullptr;
 }
 
+FName FSystemReference::GetFirstUserEventName() const
+{
+	check(IsInGameThread());
+	if (!IsValid())
+	{
+		UE_LOGFMT(LogAnimation, Warning, "GetFirstUserEventName: Invalid system handle");
+		return FName();
+	}
+
+	const FModuleEventTickFunction* TickFunction = Ptr->FindFirstUserTickFunction();
+	return TickFunction != nullptr ? TickFunction->Event.GetEventName() : FName();
+}
+
+FName FSystemReference::GetLastUserEventName() const
+{
+	check(IsInGameThread());
+	if (!IsValid())
+	{
+		UE_LOGFMT(LogAnimation, Warning, "GetLastUserEventName: Invalid system handle");
+		return FName();
+	}
+
+	const FModuleEventTickFunction* TickFunction = Ptr->FindLastUserTickFunction();
+	return TickFunction != nullptr ? TickFunction->Event.GetEventName() : FName();
+}
+
 void FSystemReference::AddDependency(UObject* InObject, FTickFunction& InTickFunction, FName InEventName, UE::UAF::ESystemDependency InDependency) const
 {
 	using namespace UE::UAF;
diff --git a/Engine/Plugins/Experimental/UAF/UAF/Source/UAF/Private/Module/UAFWeakSystemReference.cpp b/Engine/Plugins/Experimental/UAF/UAF/Source/UAF/Private/Module/UAFWeakSystemReference.cpp
--- a/Engine/Plugins/Experimental/UAF/UAF/Source/UAF/Private/Module/UAFWeakSystemReference.cpp
+++ b/Engine/Plugins/Experimental/UAF/UAF/Source/UAF/Private/Module/UAFWeakSystemReference.cpp
@@ -57,6 +57,16 @@ EPropertyBagResult FUAFWeakSystemReference::ReadVariable(const FAnimNextVariable
 	return EPropertyBagResult::PropertyNotFound;
 }
 
+bool FUAFWeakSystemReference::ReadComponent(UScriptStruct* InComponentType, TFunctionRef<void(FConstStructView InComponentStruct)> InFunction) const
+{
+	TSharedPtr<FAnimNextModuleInstance> Instance = WeakPtr.Pin();
+	if (ensure(Instance.IsValid()))
+	{
+		return Instance->ReadComponentFromExternal(InComponentType, InFunction);
+	}
+	return false;
+}
+
 const UUAFSystem* FUAFWeakSystemReference::GetSystem() const
 {
 	TSharedPtr<FAnimNextModuleInstance> Instance = WeakPtr.Pin();
@@ -65,4 +75,4 @@ const UUAFSystem* FUAFWeakSystemReference::GetSystem() const
 		return Instance->GetSystem();
 	}
 	return nullptr;
-}
+}
\ No newline at end of file
diff --git a/Engine/Plugins/Experimental/UAF/UAF/Source/UAF/Private/UAF/ValueRuntime/ValueBundleImpl_SystemOutput.cpp b/Engine/Plugins/Experimental/UAF/UAF/Source/UAF/Private/UAF/ValueRuntime/ValueBundleImpl_SystemOutput.cpp
new file mode 100644
--- /dev/null
+++ b/Engine/Plugins/Experimental/UAF/UAF/Source/UAF/Private/UAF/ValueRuntime/ValueBundleImpl_SystemOutput.cpp
@@ -0,0 +1,41 @@
+﻿// Copyright Epic Games, Inc. All Rights Reserved.
+
+#include "ValueBundleImpl_SystemOutput.h"
+
+#include "Graph/UAFSystemOutputComponent.h"
+
+namespace UE::UAF
+{
+
+const FAnimNextGraphLODPose* FValueBundleImpl_SystemOutput::GetLODPose() const
+{
+	ReadSystemOutput();
+
+	return &CachedPose;
+}
+
+const FValueBundleHeap* FValueBundleImpl_SystemOutput::GetValueBundle() const
+{
+	ReadSystemOutput();
+
+	// TODO: Implement this for value runtime
+	return &CachedValueBundle;
+}
+
+void FValueBundleImpl_SystemOutput::ReadSystemOutput() const
+{
+	// TODO: Implement this for value runtime
+	auto PoseReader = [this](uint32 InSerialNumber, const FAnimNextGraphLODPose& InPose, const FAnimNextGraphReferencePose& InRefPose)
+	{
+		SerialNumber = InSerialNumber;
+		CachedPose = InPose;
+	};
+
+	SystemReference.ReadComponent<FUAFSystemOutputComponent>([this, &PoseReader](TConstStructView<FUAFSystemOutputComponent> InStructView)
+	{
+		const FUAFSystemOutputComponent& SystemOutputComponent = InStructView.Get<FUAFSystemOutputComponent>();
+		SystemOutputComponent.ReadOutput(SerialNumber, PoseReader);
+	});
+}
+
+}
diff --git a/Engine/Plugins/Experimental/UAF/UAF/Source/UAF/Private/UAF/ValueRuntime/ValueBundleImpl_SystemOutput.h b/Engine/Plugins/Experimental/UAF/UAF/Source/UAF/Private/UAF/ValueRuntime/ValueBundleImpl_SystemOutput.h
new file mode 100644
--- /dev/null
+++ b/Engine/Plugins/Experimental/UAF/UAF/Source/UAF/Private/UAF/ValueRuntime/ValueBundleImpl_SystemOutput.h
@@ -0,0 +1,38 @@
+﻿// Copyright Epic Games, Inc. All Rights Reserved.
+
+#pragma once
+
+#include "Graph/AnimNext_LODPose.h"
+#include "Module/UAFWeakSystemReference.h"
+#include "UAF/ValueRuntime/IVirtualValueBundle.h"
+#include "UAF/ValueRuntime/ValueBundle.h"
+
+#define UE_API UAF_API
+
+namespace UE::UAF
+{
+
+struct FValueBundleImpl_SystemOutput : public IVirtualValueBundle
+{
+	explicit FValueBundleImpl_SystemOutput(const FUAFWeakSystemReference& InSystemReference)
+		: SystemReference(InSystemReference)
+	{}
+
+private:
+	// IVirtualValueBundle interface
+	virtual const FAnimNextGraphLODPose* GetLODPose() const override;
+	virtual const FValueBundleHeap* GetValueBundle() const override;
+	virtual const FWeakAnimGraphReference* GetAnimGraphReference() const override { return nullptr; }
+
+	// Read the output value bundle of the system
+	void ReadSystemOutput() const;
+
+	mutable FAnimNextGraphLODPose CachedPose;
+	mutable FValueBundleHeap CachedValueBundle;
+	FUAFWeakSystemReference SystemReference;
+	mutable uint32 SerialNumber = 0;
+};
+
+}
+
+#undef UE_API
\ No newline at end of file
diff --git a/Engine/Plugins/Experimental/UAF/UAF/Source/UAF/Private/UAFAssetInstance.cpp b/Engine/Plugins/Experimental/UAF/UAF/Source/UAF/Private/UAFAssetInstance.cpp
--- a/Engine/Plugins/Experimental/UAF/UAF/Source/UAF/Private/UAFAssetInstance.cpp
+++ b/Engine/Plugins/Experimental/UAF/UAF/Source/UAF/Private/UAFAssetInstance.cpp
@@ -123,13 +123,13 @@ FUAFAssetInstanceComponent* FUAFAssetInstance::TryGetComponent(const UScriptStru
 
 FAnimNextModuleInstance* FUAFAssetInstance::GetRootInstance()
 {
-	for (FUAFAssetInstance* CurrentHost = this; CurrentHost != nullptr; CurrentHost = CurrentHost->HostInstance)
+	for (FUAFAssetInstance* CurrentHost = this; CurrentHost != nullptr; CurrentHost = CurrentHost->HostInstance.Pin().Get())
 	{
-		if (CurrentHost->HostInstance == nullptr)
+		if (!CurrentHost->HostInstance.IsValid())
 		{
 			return CurrentHost->AsPtr<FAnimNextModuleInstance>();
 		}
 	}
 
 	return nullptr;
-}
\ No newline at end of file
+}
diff --git a/Engine/Plugins/Experimental/UAF/UAF/Source/UAF/Public/Component/AnimNextComponent.h b/Engine/Plugins/Experimental/UAF/UAF/Source/UAF/Public/Component/AnimNextComponent.h
--- a/Engine/Plugins/Experimental/UAF/UAF/Source/UAF/Public/Component/AnimNextComponent.h
+++ b/Engine/Plugins/Experimental/UAF/UAF/Source/UAF/Public/Component/AnimNextComponent.h
@@ -2,6 +2,7 @@
 
 #pragma once
 
+#include "ISystemOutputAdapter.h"
 #include "Components/ActorComponent.h"
 #include "Module/SystemReference.h"
 #include "Param/ParamType.h"
@@ -39,8 +40,41 @@ namespace UE::UAF::Editor
 	class SUAFGraphAssetPreview;
 }
 
+namespace UE::UAF::Tests
+{
+	struct FUAFSystemTests;
+}
+
+// Description of a UAF component input
+USTRUCT(meta=(Hidden))
+struct FUAFComponentInputDesc
+{
+	GENERATED_BODY()
+
+	// The component that this input uses for its input pose
+	UPROPERTY(EditAnywhere, Category="Input", meta = (UseComponentPicker, AllowedClasses = "/Script/UAF.UAFComponent"))
+	FComponentReference Component;
+
+	// The input pose variable that this input maps to
+	UPROPERTY(EditAnywhere, Category="Input", meta = (AllowedType = "FUAFValueBundle"))
+	FAnimNextVariableReference Input;
+};
+
+// The operation to perform when outputting values with respect to a skeletal mesh component
+UENUM()
+enum class EUAFSkeletalMeshComponentOutputMode : uint8
+{
+	// Write the pose to the skeletal mesh component (i.e. animate it directly)
+	WriteToSkeletalMeshComponentPose,
+
+	// Do not write the pose to the skeletal mesh component. 
+	// This can be useful to avoid unnecessary work if the skeletal mesh component's output is not rendered, or if this component is generating an 
+	// intermediate output
+	SkipWriteToSkeletalMeshComponentPose,
+};
+
 UCLASS(MinimalAPI, meta = (BlueprintSpawnableComponent))
-class UUAFComponent : public UActorComponent
+class UUAFComponent : public UActorComponent, public UE::UAF::ISystemOutputAdapter
 {
 	GENERATED_BODY()
 
@@ -55,6 +89,10 @@ class UUAFComponent : public UActorComponent
 	// UObject interface
 	virtual void Serialize(FArchive& Ar) override;
 
+	// ISystemOutputAdapter interface
+	virtual FInputData GetInputData() const override;
+	virtual void SignalOutputWritten(const UE::UAF::FModuleTaskContext& InContext) override;
+	
 public:
 	// Sets a module variable's value.
 	// @param    Name     The name of the variable to set
@@ -215,6 +253,15 @@ private:
 	// Accesses value of the specified variable for writing
 	UAF_API EPropertyBagResult WriteVariableInternal(const FAnimNextVariableReference& InVariable, const FAnimNextParamType& InType, TFunctionRef<void(TArrayView<uint8>)> InFunction);
 
+	// Determine what skeletal mesh component to write to
+	USkeletalMeshComponent* DetermineSkeletalMeshComponent() const;
+
+	// Called on registration, sets up IO
+	void SetupInputOutput();
+
+	// Called on registration, or lazily on dependent registration, as component registration order is not guaranteed 
+	void AllocateSystem();
+
 private:
 	friend struct UE::UAF::UncookedOnly::FUtils;
 	friend struct FAssetBlueprintTests;
@@ -224,6 +271,7 @@ private:
 	friend class UE::UAF::Editor::FSystemViewportController;
 	friend class UE::UAF::Editor::FAnimGraphViewportController;
 	friend class UE::UAF::Editor::SUAFGraphAssetPreview;
+	friend UE::UAF::Tests::FUAFSystemTests;
 
 	// The asset that this component will run
 #if WITH_EDITORONLY_DATA
@@ -232,7 +280,7 @@ private:
 	TObjectPtr<UObject> Asset_DEPRECATED = nullptr;
 #endif // WITH_EDITORONLY_DATA
 
-	UPROPERTY(EditAnywhere, Category="Asset", meta=(ExcludeBaseStruct))
+	UPROPERTY(EditAnywhere, Category="Animation", meta=(ExcludeBaseStruct))
 	TInstancedStruct<FUAFSystemFactoryAsset> AssetData;
 
 #if WITH_EDITORONLY_DATA
@@ -240,14 +288,38 @@ private:
 	TObjectPtr<UUAFSystem> Module_DEPRECATED = nullptr;
 #endif
 
+	// The inputs to this UAF system. These map the outputs of other UAF components to the specified input values of this component's current system.
+	// ~ Note that these are EditDefaultsOnly to work around limitations with array editing for component instances. If we want to make these
+	// ~ per-instance-editable, then we may need a custom handler for component instance data to support that.
+	UPROPERTY(EditDefaultsOnly, Category="Input")
+	TArray<FUAFComponentInputDesc> Inputs;
+
+	// The skeletal mesh component that this UAF system will animate against.
+	// If this is not set, then the first component that is found on the actor will be used.
+	// Note that setting this does not necessarily mean that the mesh's pose will be written by the UAF system, only that the output value will match
+	// the skeletal mesh component's 'shape'.
+	// ~ Note that this is EditDefaultsOnly to complement Inputs.
+	UPROPERTY(EditDefaultsOnly, Category="Output", meta = (UseComponentPicker, AllowedClasses = "/Script/Engine.SkeletalMeshComponent"))
+	FComponentReference OutputComponent;
+
+	// The operation to perform when outputting values with respect to a skeletal mesh component
+	UPROPERTY(EditDefaultsOnly, Category="Output")
+	EUAFSkeletalMeshComponentOutputMode SkeletalMeshComponentOutputMode = EUAFSkeletalMeshComponentOutputMode::WriteToSkeletalMeshComponentPose;
+
+	UPROPERTY()
+	TObjectPtr<USkeletalMeshComponent> CachedSkeletalMeshComponent;
+
 	// Owning reference to the registered system
 	UE::UAF::FSystemReference SystemReference;
 
+	// Serial number, updated on worker threads when output is written
+	uint32 SerialNumber = 0;
+
 	// How to initialize the system
-	UPROPERTY(EditAnywhere, Category="System")
+	UPROPERTY(EditAnywhere, Category="Animation")
 	EAnimNextModuleInitMethod InitMethod = EAnimNextModuleInitMethod::InitializeAndPauseInEditor;
 
 	/** When checked, the system's debug drawing instructions are drawn in the viewport */
-	UPROPERTY(EditAnywhere, Category = "Rendering")
+	UPROPERTY(EditAnywhere, Category = "Animation")
 	uint8 bShowDebugDrawing : 1 = false;
 };
diff --git a/Engine/Plugins/Experimental/UAF/UAF/Source/UAF/Public/IPoseAdapter.h b/Engine/Plugins/Experimental/UAF/UAF/Source/UAF/Public/IPoseAdapter.h
deleted file mode 100644
diff --git a/Engine/Plugins/Experimental/UAF/UAF/Source/UAF/Public/ISystemOutputAdapter.h b/Engine/Plugins/Experimental/UAF/UAF/Source/UAF/Public/ISystemOutputAdapter.h
new file mode 100644
--- /dev/null
+++ b/Engine/Plugins/Experimental/UAF/UAF/Source/UAF/Public/ISystemOutputAdapter.h
@@ -0,0 +1,66 @@
+﻿// Copyright Epic Games, Inc. All Rights Reserved.
+
+#pragma once
+
+#include "CoreTypes.h"
+#include "UObject/ObjectPtr.h"
+#include "Components/SkeletalMeshComponent.h"
+
+class USkeletalMesh;
+class USkeletalMeshComponent;
+struct FAnimNextGraphLODPose;
+struct FAnimNextGraphReferencePose;
+struct FUAFSystemOutputComponent;
+
+namespace UE::UAF
+{
+struct FModuleTaskContext;
+}
+
+namespace UE::UAF
+{
+
+// Interface used by any host that wants to read an output from a UAF system
+// Register this with FUAFSystemOutputComponent
+class ISystemOutputAdapter
+{
+public:
+	virtual ~ISystemOutputAdapter() = default;
+	
+	// All input data required to generate an output value
+	struct FInputData
+	{
+		FInputData(TNonNullPtr<USkeletalMesh> InSkeletalMesh, int32 InLOD)
+			: SkeletalMesh(InSkeletalMesh)
+			, LOD(InLOD)
+		{}
+
+		FInputData(TNonNullPtr<USkeletalMeshComponent> InSkeletalMeshComponent)
+			: SkeletalMesh(InSkeletalMeshComponent->GetSkeletalMeshAsset())
+			, SkeletalMeshComponent(InSkeletalMeshComponent)
+			, LOD(InSkeletalMeshComponent->GetPredictedLODLevel())
+		{}
+
+		// Mesh to use, can be nullptr when using SkeletalMeshComponent (as the component can have a null mesh), but must be valid when not
+		USkeletalMesh* SkeletalMesh = nullptr;
+
+		// Legacy mesh component to use, can be nullptr
+		USkeletalMeshComponent* SkeletalMeshComponent = nullptr;
+
+		// LOD to run at
+		int32 LOD = 0;
+	};
+
+private:
+	friend struct ::FUAFSystemOutputComponent;
+
+	// Called back on a worker thread by FUAFSystemOutputComponent to retrieve the LOD and mesh to use when system execution begins
+	// The mesh & LOD is then consistent across a system's execution
+	// Consistency is validated with an ensure() in FUAFSystemOutputComponent prior to calling back via GenerateRenderData
+	virtual FInputData GetInputData() const = 0;
+
+	// Called back on a worker thread by FUAFSystemOutputComponent when an output is written
+	virtual void SignalOutputWritten(const UE::UAF::FModuleTaskContext& InContext) = 0;
+};
+
+}
diff --git a/Engine/Plugins/Experimental/UAF/UAF/Source/UAF/Public/InstanceTaskContext.h b/Engine/Plugins/Experimental/UAF/UAF/Source/UAF/Public/InstanceTaskContext.h
--- a/Engine/Plugins/Experimental/UAF/UAF/Source/UAF/Public/InstanceTaskContext.h
+++ b/Engine/Plugins/Experimental/UAF/UAF/Source/UAF/Public/InstanceTaskContext.h
@@ -36,9 +36,9 @@ public:
 	template<typename ValueType>
 	bool AccessVariable(const FAnimNextVariableReference& InVariable, TFunctionRef<void(ValueType&)> InFunction) const
 	{
-		return AccessVariableInternal(InVariable, [&InFunction](TArrayView<uint8> InArrayView)
+		return AccessVariableInternal(InVariable, FAnimNextParamType::GetType<ValueType>(), [&InFunction](TArrayView<uint8> InArrayView)
 		{
-			InFunction(*reinterpret_cast<ValueType>(InArrayView.GetData()));
+			InFunction(*reinterpret_cast<ValueType*>(InArrayView.GetData()));
 		});
 	}
 
diff --git a/Engine/Plugins/Experimental/UAF/UAF/Source/UAF/Public/Module/ModuleTaskContext.h b/Engine/Plugins/Experimental/UAF/UAF/Source/UAF/Public/Module/ModuleTaskContext.h
--- a/Engine/Plugins/Experimental/UAF/UAF/Source/UAF/Public/Module/ModuleTaskContext.h
+++ b/Engine/Plugins/Experimental/UAF/UAF/Source/UAF/Public/Module/ModuleTaskContext.h
@@ -11,7 +11,7 @@ struct FUAFAssetInstanceComponent;
 struct FAnimNextModuleInstance;
 struct FInstancedPropertyBag;
 struct FAnimNextModuleInjectionComponent;
-struct FUAFOutputPoseComponent;
+struct FUAFSystemOutputComponent;
 
 namespace UE::UAF
 {
@@ -38,6 +38,11 @@ public:
 	// Access a module instance component of the specified type. If the component exists, then InFunction will be called
 	UAF_API void TryAccessComponent(UScriptStruct* InComponentType, TFunctionRef<void(FUAFAssetInstanceComponent&)> InFunction) const;
 
+	// Access a module instance component of the specified type. 
+	// If the component does not exist, it will be created.
+	UAF_API void AccessComponent(UScriptStruct* InComponentType, TFunctionRef<void(FUAFAssetInstanceComponent&)> InFunction) const;
+
+	// Access a module instance component of the specified type. If the component exists, then InFunction will be called
 	template<typename ComponentType>
 	void TryAccessComponent(TFunctionRef<void(ComponentType&)> InFunction) const
 	{
@@ -46,6 +51,17 @@ public:
 			InFunction(static_cast<ComponentType&>(InComponent));
 		});
 	}
+	
+	// Access a module instance component of the specified type. 
+	// If the component does not exist, it will be created.
+	template<typename ComponentType>
+	void AccessComponent(TFunctionRef<void(ComponentType&)> InFunction) const
+	{
+		AccessComponent(TBaseStructure<ComponentType>::Get(), [&InFunction](FUAFAssetInstanceComponent& InComponent)
+		{
+			InFunction(static_cast<ComponentType&>(InComponent));
+		});
+	}
 
 	UAF_API FAnimNextModuleInstance* const GetModuleInstance() const;
 
@@ -57,7 +73,9 @@ private:
 
 	friend UE::UAF::FModuleEventTickFunction;
 	friend ::FAnimNextModuleInjectionComponent;
-	friend ::FUAFOutputPoseComponent;
+	friend ::FUAFSystemOutputComponent;
 };
 
+using FUniqueSystemTask = TUniqueFunction<void(const UE::UAF::FModuleTaskContext&)>;
+
 }
\ No newline at end of file
diff --git a/Engine/Plugins/Experimental/UAF/UAF/Source/UAF/Public/Module/SystemReference.h b/Engine/Plugins/Experimental/UAF/UAF/Source/UAF/Public/Module/SystemReference.h
--- a/Engine/Plugins/Experimental/UAF/UAF/Source/UAF/Public/Module/SystemReference.h
+++ b/Engine/Plugins/Experimental/UAF/UAF/Source/UAF/Public/Module/SystemReference.h
@@ -111,6 +111,14 @@ private:
 	UE_API void ShowDebugDrawing(bool bInShowDebugDrawing) const;
 #endif
 
+	// Queue a task to run before the first user event
+	// Can be called on any thread.
+	// @param	InTaskFunction		The function to run
+	void QueueTask(TUniqueFunction<void(const UE::UAF::FModuleTaskContext&)>&& InTaskFunction) const
+	{
+		QueueTask(FName(), MoveTemp(InTaskFunction), UE::UAF::ETaskRunLocation::Before);
+	}
+	
 	// Queue a task to run at a particular point in a system's execution
 	// Can be called on any thread.
 	// @param	InSystemEventName	The name of the event in the system to run the supplied task relative to. If this is NAME_None, then the first user event will be used.
@@ -128,6 +136,14 @@ private:
 	// @param	InEventName			The event associated to the wanted tick function
 	UE_API const FTickFunction* FindTickFunction(FName InEventName) const;
 
+	// Get the event name of the first user event
+	// Should be called only from the game thread
+	UE_API FName GetFirstUserEventName() const;
+
+	// Get the event name of the last user event
+	// Should be called only from the game thread
+	UE_API FName GetLastUserEventName() const;
+
 	// Run a specific Event in-place.
 	// @TODO: Should only be called on unregistered tick functions.
 	// @param	InEventName		The event we want to run.
@@ -184,6 +200,7 @@ private:
 	UE_API EPropertyBagResult WriteVariable(const FAnimNextVariableReference& InVariable, const FAnimNextParamType& InType, TFunctionRef<void(TArrayView<uint8>)> InFunction) const;
 
 	// Access the supplied system component for reading.
+	// Can be called on any thread.
 	// @param	InFunction			Function that will be called to access the struct
 	// @return true if the struct was available, otherwise false
 	template<typename ComponentType>
diff --git a/Engine/Plugins/Experimental/UAF/UAF/Source/UAF/Public/Module/UAFWeakSystemReference.h b/Engine/Plugins/Experimental/UAF/UAF/Source/UAF/Public/Module/UAFWeakSystemReference.h
--- a/Engine/Plugins/Experimental/UAF/UAF/Source/UAF/Public/Module/UAFWeakSystemReference.h
+++ b/Engine/Plugins/Experimental/UAF/UAF/Source/UAF/Public/Module/UAFWeakSystemReference.h
@@ -5,11 +5,13 @@
 #include "AnimNextPoolHandle.h"
 #include "ModuleTaskContext.h"
 #include "TaskRunLocation.h"
+#include "UAFModuleInstanceComponent.h"
 #include "UAFWeakSystemReference.generated.h"
 
 #define UE_API UAF_API
 
 class UUAFSystem;
+class UUAFAnimGraph;
 struct FAnimNextModuleInstance;
 
 namespace UE::UAF
@@ -43,6 +45,14 @@ struct FUAFWeakSystemReference
 		return WeakPtr.IsValid();
 	}
 
+	// Queue a task to run before the first user event
+	// Can be called on any thread.
+	// @param	InTaskFunction		The function to run
+	void QueueTask(TUniqueFunction<void(const UE::UAF::FModuleTaskContext&)>&& InTaskFunction) const
+	{
+		QueueTask(FName(), MoveTemp(InTaskFunction), UE::UAF::ETaskRunLocation::Before);
+	}
+
 	// Queue a task to run at a particular point in a system's execution
 	// Can be called on any thread.
 	// @param	InSystemEventName	The name of the event in the system to run the supplied task relative to. If this is NAME_None, then the first user event will be used.
@@ -79,12 +89,27 @@ struct FUAFWeakSystemReference
 	// This performs non-recursive locks on the system instance to ensure that concurrent access is respected, so care must be taken calling this from
 	// animation worker tasks as deadlocks can occur.
 	// Type must match strictly, no conversions are performed.
+	// Can be called on any thread.
 	// @param	InVariable			The variable to get the value of
 	// @param	InType				The type of the variable
 	// @param	InFunction			Function that will be called if no errors occur
 	// @return see EPropertyBagResult
 	UE_API EPropertyBagResult ReadVariable(const FAnimNextVariableReference& InVariable, const FAnimNextParamType& InType, TFunctionRef<void(TConstArrayView<uint8>)> InFunction) const;
 
+	// Access the supplied system component for reading.
+	// Can be called on any thread.
+	// @param	InFunction			Function that will be called to access the struct
+	// @return true if the struct was available, otherwise false
+	template<typename ComponentType>
+	bool ReadComponent(TFunctionRef<void(TConstStructView<ComponentType> InComponentStruct)> InFunction) const
+	{
+		static_assert(std::is_base_of_v<FUAFModuleInstanceComponent, ComponentType>, "Component struct must be a child of FUAFModuleInstanceComponent");
+		return ReadComponent(ComponentType::StaticStruct(), [&InFunction](FConstStructView InStructView)
+		{
+			InFunction(TConstStructView<ComponentType>(InStructView.Get<ComponentType>()));
+		});
+	}
+
 	UE_API const UUAFSystem* GetSystem() const;
 
 private:
@@ -92,6 +117,9 @@ private:
 	friend FAnimNextModuleInstance;
 	friend UE::UAF::FModuleEventTickFunction;
 	friend UE::UAF::FInjectionRequest;
+	friend UUAFAnimGraph;
+
+	UE_API bool ReadComponent(UScriptStruct* InComponentType, TFunctionRef<void(FConstStructView InComponentStruct)> InFunction) const;
 
 	// Ptr to the instance
 	TWeakPtr<FAnimNextModuleInstance> WeakPtr;
diff --git a/Engine/Plugins/Experimental/UAF/UAF/Source/UAF/Public/UAF/ValueRuntime/ValueBundle.h b/Engine/Plugins/Experimental/UAF/UAF/Source/UAF/Public/UAF/ValueRuntime/ValueBundle.h
--- a/Engine/Plugins/Experimental/UAF/UAF/Source/UAF/Public/UAF/ValueRuntime/ValueBundle.h
+++ b/Engine/Plugins/Experimental/UAF/UAF/Source/UAF/Public/UAF/ValueRuntime/ValueBundle.h
@@ -14,6 +14,12 @@
 
 #define UE_API UAF_API
 
+namespace UE::UAF
+{
+	struct FInputValueTrait;
+	class IVirtualValueBundle;
+};
+
 namespace UE::UAF
 {
 	/*
@@ -121,21 +127,55 @@ namespace UE::UAF
 }
 
 /*
- * Attribute Collection
+ * Value Bundle
  *
+ * 
  * This struct holds attributes produced by UAF. They come in two flavors: static (set map) and dynamic (map).
  * Static attributes are fast and defined up front in the skeleton set binding asset.
  * Dynamic attributes are slower but produced on demand during evaluation.
+ * This struct can also represent a 'virtual value', for instance an inlined, or owned graph instance 
  */
 USTRUCT(BlueprintType, meta = (DisplayName = "Value Bundle"))
 struct FUAFValueBundle
-#if CPP
-	: public UE::UAF::FValueBundleHeap
-#endif
 {
 	GENERATED_BODY()
 
-	// Equivalent to FAnimNextGraphLODPose
+	// Set as a reference pose (the assumed value of an empty value bundle)
+	void SetAsRefPose()
+	{
+		Impl.Reset();
+	}
+
+	// Get whether this is a reference pose
+	bool IsRefPose() const
+	{
+		return !Impl.IsValid();
+	}
+
+	// Set as an implementation - @see UE::UAF::IVirtualValueBundle
+	template<typename ImplType, typename... ArgsType>
+	void SetAs(ArgsType&&... InArgs)
+	{
+		static_assert(std::is_base_of_v<UE::UAF::IVirtualValueBundle, ImplType>, "Value Bundle Implementation must be a child of UE::UAF::IVirtualValueBundle");
+		Impl = MakeShared<ImplType>(Forward<ArgsType>(InArgs)...);
+	}
+
+	// Access the implementation
+	const UE::UAF::IVirtualValueBundle* GetImpl() const
+	{
+		return Impl.Get();
+	}
+
+	// Access the implementation, asserting if it is invalid/refpose
+	UE::UAF::IVirtualValueBundle& GetImplChecked() const
+	{
+		check(Impl.IsValid());
+		return *Impl.Get();
+	}
+
+private:
+	// Implementation
+	TSharedPtr<UE::UAF::IVirtualValueBundle> Impl;
 };
 
 //////////////////////////////////////////////////////////////////////////
diff --git a/Engine/Plugins/Experimental/UAF/UAF/Source/UAF/Public/UAFAssetInstance.h b/Engine/Plugins/Experimental/UAF/UAF/Source/UAF/Public/UAFAssetInstance.h
--- a/Engine/Plugins/Experimental/UAF/UAF/Source/UAF/Public/UAFAssetInstance.h
+++ b/Engine/Plugins/Experimental/UAF/UAF/Source/UAF/Public/UAFAssetInstance.h
@@ -8,6 +8,7 @@
 #include "Param/ParamType.h"
 #include "StructUtils/PropertyBag.h"
 #include "RewindDebugger/UAFTrace.h"
+#include "Templates/SubScriptStructOf.h"
 #include "Variables/UAFInstanceVariableData.h"
 
 #include "UAFAssetInstance.generated.h"
@@ -28,7 +29,7 @@ namespace UE::UAF
 	struct FUAFTrace;
 	struct FInstanceTaskContext;
 	struct FVariableOverridesCollection;
-	struct FEntitySkeletalAnimationStackTrait;
+	struct FLayerStackTrait;
 }
 
 // Base struct for instances of UAF assets
@@ -168,6 +169,25 @@ struct FUAFAssetInstance
 		return NewComponent;
 	}
 
+	// Returns an instance component of the requested type, creating it lazily the first time it is queried.
+	// When querying for structs of a base type, existing derived types will be returned, rather than a new base struct created.
+	[[nodiscard]] FUAFAssetInstanceComponent& GetComponent(TSubScriptStructOf<FUAFAssetInstanceComponent> InScriptStruct)
+	{
+		check(InScriptStruct.Get() != nullptr);
+		if (FUAFAssetInstanceComponent* ExistingComponent = TryGetComponent(InScriptStruct))
+		{
+			return *ExistingComponent;
+		}
+
+		TInstancedStruct<FUAFAssetInstanceComponent> NewInstancedStruct;
+		NewInstancedStruct.InitializeAsScriptStruct(InScriptStruct);
+		TInstancedStruct<FUAFAssetInstanceComponent>& InstancedStruct = Components.Add_GetRef(MoveTemp(NewInstancedStruct));
+		FUAFAssetInstanceComponent& NewComponent = InstancedStruct.GetMutable<FUAFAssetInstanceComponent>();
+		NewComponent.Instance = this;
+		NewComponent.OnBindToInstance();
+		return NewComponent;
+	}
+
 	// Returns a typed instance component pointer if found or nullptr otherwise.
 	// When querying for structs of a base type, existing derived types will be returned.
 	template<typename ComponentType>
@@ -211,7 +231,7 @@ struct FUAFAssetInstance
 	// Get the instance (graph, module etc.) that owns/hosts us
 	FUAFAssetInstance* GetHost() const
 	{
-		return HostInstance;
+		return HostInstance.Pin().Get();
 	}
 
 	// Get the root module instance that owns us. Can be nullptr in some cases, such as automated tests.
@@ -329,7 +349,7 @@ protected:
 	TObjectPtr<const UScriptStruct> ScriptStruct;
 
 	// The instance (graph, module etc.) that owns/hosts us
-	FUAFAssetInstance* HostInstance = nullptr;
+	TWeakPtr<FUAFAssetInstance> HostInstance;
 
 	friend UE::UAF::FInjectionInfo;
 	friend FAnimNextModuleInjectionComponent;
@@ -338,7 +358,7 @@ protected:
 	friend UE::UAF::FUAFTrace;
 	friend UE::UAF::FInstanceTaskContext;
 	friend FUAFRigVMComponent;
-	friend UE::UAF::FEntitySkeletalAnimationStackTrait;
+	friend UE::UAF::FLayerStackTrait;
 
 	FName GetDebugName() const
 	{
diff --git a/Engine/Plugins/Experimental/UAF/UAF/Source/UAFUncookedOnly/Internal/AnimNextRigVMAssetEditorData.h b/Engine/Plugins/Experimental/UAF/UAF/Source/UAFUncookedOnly/Internal/AnimNextRigVMAssetEditorData.h
--- a/Engine/Plugins/Experimental/UAF/UAF/Source/UAFUncookedOnly/Internal/AnimNextRigVMAssetEditorData.h
+++ b/Engine/Plugins/Experimental/UAF/UAF/Source/UAFUncookedOnly/Internal/AnimNextRigVMAssetEditorData.h
@@ -86,6 +86,7 @@ namespace UE::UAF::Tests
 	class FVariables;
 	class FVariables_UOLBindings;
 	class FDataInterfaceCompile;
+	struct FUAFSystemTests;
 }
 
 enum class EAnimNextEditorDataNotifType : uint8
@@ -247,7 +248,8 @@ protected:
 	friend class UE::UAF::Editor::FFunctionsOutlinerHierarchy;	
 	friend class UE::UAF::Editor::FOutlinerHierarchy;
 	friend class UE::UAF::Editor::SVariablesOutlinerValue;
-	friend class UE::UAF::Editor::FFunctionsOutlinerMode;	
+	friend class UE::UAF::Editor::FFunctionsOutlinerMode;
+	friend UE::UAF::Tests::FUAFSystemTests;
 	friend class UAnimNextModuleWorkspaceAssetUserData;
 	friend class UE::UAF::Editor::SAddVariablesDialog;
 	friend class UE::UAF::Editor::FVariableProxyCustomization;
@@ -275,6 +277,7 @@ protected:
 	friend UE::UAF::Editor::SVariablesView;
 	friend UUAFSharedVariableNode;
 	friend class UE::UAF::UncookedOnly::FModule;
+	
 
 	// UObject interface
 	UE_API virtual void Serialize(FArchive& Ar) override;
diff --git a/Engine/Plugins/Experimental/UAF/UAFAnimGraph/Source/UAFAnimGraph/Internal/Graph/InputPoseImpl_ExternalGraph.h b/Engine/Plugins/Experimental/UAF/UAFAnimGraph/Source/UAFAnimGraph/Internal/Graph/InputPoseImpl_ExternalGraph.h
new file mode 100644
--- /dev/null
+++ b/Engine/Plugins/Experimental/UAF/UAFAnimGraph/Source/UAFAnimGraph/Internal/Graph/InputPoseImpl_ExternalGraph.h
@@ -0,0 +1,34 @@
+﻿// Copyright Epic Games, Inc. All Rights Reserved.
+
+#pragma once
+
+#include "Graph/WeakAnimGraphReference.h"
+#include "UAF/ValueRuntime/IVirtualValueBundle.h"
+
+#define UE_API UAFANIMGRAPH_API
+
+namespace UE::UAF
+{
+
+struct FValueBundleImpl_ExternalGraph : public IVirtualValueBundle
+{
+	explicit FValueBundleImpl_ExternalGraph(const FWeakAnimGraphReference& InGraphReference)
+		: WeakGraphReference(InGraphReference)
+	{}
+
+	explicit FValueBundleImpl_ExternalGraph(const FAnimGraphReference& InGraphReference)
+		: WeakGraphReference(InGraphReference)
+	{}
+
+private:
+	// IVirtualValueBundle interface
+	UE_API virtual const FAnimNextGraphLODPose* GetLODPose() const override;
+	UE_API virtual const FValueBundleHeap* GetValueBundle() const override;
+	UE_API virtual const FWeakAnimGraphReference* GetAnimGraphReference() const override;
+
+	FWeakAnimGraphReference WeakGraphReference;
+};
+
+}
+
+#undef UE_API
\ No newline at end of file
diff --git a/Engine/Plugins/Experimental/UAF/UAFAnimGraph/Source/UAFAnimGraph/Internal/Graph/RigUnit_AnimNextRunAnimationGraph_v2.h b/Engine/Plugins/Experimental/UAF/UAFAnimGraph/Source/UAFAnimGraph/Internal/Graph/RigUnit_AnimNextRunAnimationGraph_v2.h
--- a/Engine/Plugins/Experimental/UAF/UAFAnimGraph/Source/UAFAnimGraph/Internal/Graph/RigUnit_AnimNextRunAnimationGraph_v2.h
+++ b/Engine/Plugins/Experimental/UAF/UAFAnimGraph/Source/UAFAnimGraph/Internal/Graph/RigUnit_AnimNextRunAnimationGraph_v2.h
@@ -67,6 +67,9 @@ struct FRigUnit_AnimNextRunAnimationGraph_v2 : public FRigUnit_AnimNextBase
 	// The execution result
 	UPROPERTY(EditAnywhere, DisplayName = "Execute", Category = "BeginExecution", meta = (Input, Output))
 	FAnimNextExecuteContext ExecuteContext;
+	
+	RIGVM_METHOD()
+	UE_API virtual FRigVMStructUpgradeInfo GetUpgradeInfo() const override;
 };
 
 #undef UE_API
diff --git a/Engine/Plugins/Experimental/UAF/UAFAnimGraph/Source/UAFAnimGraph/Internal/Graph/RigUnit_UAFRunAsset.h b/Engine/Plugins/Experimental/UAF/UAFAnimGraph/Source/UAFAnimGraph/Internal/Graph/RigUnit_UAFRunAsset.h
new file mode 100644
--- /dev/null
+++ b/Engine/Plugins/Experimental/UAF/UAFAnimGraph/Source/UAFAnimGraph/Internal/Graph/RigUnit_UAFRunAsset.h
@@ -0,0 +1,65 @@
+// Copyright Epic Games, Inc. All Rights Reserved.
+
+#pragma once
+
+#include "AnimNextExecuteContext.h"
+#include "Graph/AnimNextAnimGraph.h"
+#include "Graph/AnimNext_LODPose.h"
+#include "Graph/RigUnit_AnimNextBase.h"
+#include "UAF/ValueRuntime/ValueBundle.h"
+#include "Variables/AnimNextVariableOverridesCollection.h"
+
+#include "RigUnit_UAFRunAsset.generated.h"
+
+#define UE_API UAFANIMGRAPH_API
+
+USTRUCT()
+struct FUAFRunAssetWorkData
+{
+	GENERATED_BODY()
+
+	// Weak ptr to the instance we run (ownership is with the module component)
+	TWeakPtr<FAnimNextGraphInstance> WeakHost;
+
+	// Graph we are injecting into the host
+	UPROPERTY(transient)
+	FAnimNextVariableReference InjectedGraphReference;
+};
+
+/** Runs a UAF asset to generate an output value */
+USTRUCT(meta=(DisplayName="Run Asset", Category="Anim", NodeColor="0, 1, 1", Keywords="Trait,Stack", RequiredComponents="AnimNextModuleInjectionComponent"))
+struct FRigUnit_UAFRunAsset : public FRigUnit_AnimNextBase
+{
+	GENERATED_BODY()
+
+	RIGVM_METHOD()
+	UE_API void Execute();
+
+	static UE_API const UUAFAnimGraph* GetHostGraphToRun(FAnimNextExecuteContext& InExecuteContext);
+
+	// The graph to run
+	UPROPERTY(EditAnywhere, Category = "Run", meta = (Input, ShowOnlyInnerProperties, ExportAsReference="true"))
+	FAnimNextAnimGraph Graph;
+
+	// Variable overrides to be applied to subgraphs
+	UPROPERTY(EditAnywhere, Category = "Run", meta = (Input))
+	FAnimNextVariableOverridesCollection Overrides;
+
+	// The name of the attribute set to execute with, if 'None' will include everything
+	UPROPERTY(EditAnywhere, Category = "Run", meta = (Input))
+	FName AttributeSet;
+
+	// Result value
+	UPROPERTY(EditAnywhere, Category = "Run", meta = (Output))
+	FUAFValueBundle Result;
+
+	// Internal work data
+	UPROPERTY(transient)
+	FUAFRunAssetWorkData WorkData;
+
+	// The execution result
+	UPROPERTY(EditAnywhere, DisplayName = "Execute", Category = "BeginExecution", meta = (Input, Output))
+	FAnimNextExecuteContext ExecuteContext;
+};
+
+#undef UE_API
diff --git a/Engine/Plugins/Experimental/UAF/UAFAnimGraph/Source/UAFAnimGraph/Internal/Traits/InputValueTrait.h b/Engine/Plugins/Experimental/UAF/UAFAnimGraph/Source/UAFAnimGraph/Internal/Traits/InputValueTrait.h
new file mode 100644
--- /dev/null
+++ b/Engine/Plugins/Experimental/UAF/UAFAnimGraph/Source/UAFAnimGraph/Internal/Traits/InputValueTrait.h
@@ -0,0 +1,61 @@
+// Copyright Epic Games, Inc. All Rights Reserved.
+
+#pragma once
+
+#include "CoreMinimal.h"
+
+#include "TraitCore/Trait.h"
+#include "TraitInterfaces/IEvaluate.h"
+#include "TraitInterfaces/IHierarchy.h"
+#include "TraitInterfaces/IUpdate.h"
+
+#include "InputValueTrait.generated.h"
+
+USTRUCT(meta = (DisplayName = "Input Value", ShowTooltip=true))
+struct FUAFInputValueTraitSharedData : public FAnimNextTraitSharedData
+{
+	GENERATED_BODY()
+
+	// The variable to use as an input value
+	UPROPERTY(EditAnywhere, Category="Input Value", meta=(Input, HideSubPins, AllowedType = FUAFValueBundle))
+	FAnimNextVariableReference Input;
+
+	// Latent pin support boilerplate
+#define TRAIT_LATENT_PROPERTIES_ENUMERATOR(GeneratorMacro) \
+	GeneratorMacro(Input)
+
+	GENERATE_TRAIT_LATENT_PROPERTIES(FUAFInputValueTraitSharedData, TRAIT_LATENT_PROPERTIES_ENUMERATOR)
+#undef TRAIT_LATENT_PROPERTIES_ENUMERATOR
+};
+
+namespace UE::UAF
+{
+	struct FInputValueTrait : FBaseTrait, IUpdate, IEvaluate, IUpdateTraversal, IHierarchy
+	{
+		DECLARE_ANIM_TRAIT(FInputValueTrait, FBaseTrait)
+
+		using FSharedData = FUAFInputValueTraitSharedData;
+
+		struct FInstanceData : FTrait::FInstanceData
+		{
+			FAnimNextVariableReference VariableReference;
+			FAnimNextGraphInstance* GraphInstance = nullptr;
+		};
+
+		// IUpdate impl
+		virtual void PreUpdate(FUpdateTraversalContext& Context, const TTraitBinding<IUpdate>& Binding, const FTraitUpdateState& TraitState) const override;
+		virtual void OnBecomeRelevant(FUpdateTraversalContext& Context, const TTraitBinding<IUpdate>& Binding, const FTraitUpdateState& TraitState) const override;
+
+		// IEvaluate impl
+		virtual void PreEvaluate(FEvaluateTraversalContext& Context, const TTraitBinding<IEvaluate>& Binding) const override;
+
+		// IUpdateTraversal impl
+		virtual void QueueChildrenForTraversal(FUpdateTraversalContext& Context, const TTraitBinding<IUpdateTraversal>& Binding, const FTraitUpdateState& TraitState, FUpdateTraversalQueue& TraversalQueue) const override;
+
+		// IHierarchy impl
+		virtual uint32 GetNumChildren(const FExecutionContext& Context, const TTraitBinding<IHierarchy>& Binding) const override;
+		virtual void GetChildren(const FExecutionContext& Context, const TTraitBinding<IHierarchy>& Binding, FChildrenArray& Children) const override;
+
+		const FInstanceData* UpdateInstanceData(const FTraitBinding& InBinding) const;
+	};
+}
\ No newline at end of file
diff --git a/Engine/Plugins/Experimental/UAF/UAFAnimGraph/Source/UAFAnimGraph/Private/Graph/AnimGraph/AnimNode_AnimNextGraph.cpp b/Engine/Plugins/Experimental/UAF/UAFAnimGraph/Source/UAFAnimGraph/Private/Graph/AnimGraph/AnimNode_AnimNextGraph.cpp
--- a/Engine/Plugins/Experimental/UAF/UAFAnimGraph/Source/UAFAnimGraph/Private/Graph/AnimGraph/AnimNode_AnimNextGraph.cpp
+++ b/Engine/Plugins/Experimental/UAF/UAFAnimGraph/Source/UAFAnimGraph/Private/Graph/AnimGraph/AnimNode_AnimNextGraph.cpp
@@ -63,7 +63,6 @@ void FAnimNode_AnimNextGraph::Update_AnyThread(const FAnimationUpdateContext& Co
 		PropagateInputProperties(Context.AnimInstanceProxy->GetAnimInstanceObject());
 
 		FUpdateGraphContext UpdateGraphContext(*GraphInstance.Get(), Context.GetDeltaTime());
-		UpdateGraphContext.SetBindingObject(Context.AnimInstanceProxy->GetSkelMeshComponent());
 		UE::UAF::UpdateGraph(UpdateGraphContext);
 	}
 
@@ -126,7 +125,6 @@ void FAnimNode_AnimNextGraph::Evaluate_AnyThread(FPoseContext& Output)
 
 		{
 			UE::UAF::FEvaluateGraphContext EvaluateGraphContext(*GraphInstance.Get(), RefPose, LODLevel);
-			EvaluateGraphContext.SetBindingObject(SkeletalMeshComponent);
 
 			if (UE::UAF::UseAbstractSkeletonRuntime())
 			{
diff --git a/Engine/Plugins/Experimental/UAF/UAFAnimGraph/Source/UAFAnimGraph/Private/Graph/AnimGraphReference.cpp b/Engine/Plugins/Experimental/UAF/UAFAnimGraph/Source/UAFAnimGraph/Private/Graph/AnimGraphReference.cpp
new file mode 100644
--- /dev/null
+++ b/Engine/Plugins/Experimental/UAF/UAFAnimGraph/Source/UAFAnimGraph/Private/Graph/AnimGraphReference.cpp
@@ -0,0 +1,81 @@
+﻿// Copyright Epic Games, Inc. All Rights Reserved.
+
+#include "Graph/AnimGraphReference.h"
+#include "Graph/AnimNextAnimationGraph.h"
+#include "Graph/AnimNextGraphInstance.h"
+#include "Graph/WeakAnimGraphReference.h"
+#include "Module/AnimNextModuleInstance.h"
+#include "TraitCore/ExecutionContext.h"
+#include "TraitCore/NodeInstance.h"
+
+namespace UE::UAF
+{
+
+FAnimGraphReference::FAnimGraphReference(TNonNullPtr<const UUAFAnimGraph> InAnimGraph, const FUAFWeakSystemReference& InSystemReference, FAnimNextFactoryParams&& InFactoryParams)
+{
+	check(InSystemReference.IsValid());
+	Ptr = InAnimGraph->PreAllocateInstance({ .SystemReference = InSystemReference }, MoveTemp(InFactoryParams));
+}
+
+FWeakAnimGraphReference FAnimGraphReference::ToWeakRef() const
+{
+	return FWeakAnimGraphReference(*this);
+}
+
+void FAnimGraphReference::Reset()
+{
+	Ptr.Reset();
+}
+
+bool FAnimGraphReference::IsValid() const
+{
+	return Ptr.IsValid();
+}
+
+const UUAFAnimGraph* FAnimGraphReference::GetAnimGraph() const
+{
+	return Ptr.IsValid() ? Ptr->GetAnimationGraph() : nullptr;
+}
+
+TSharedPtr<FAnimNextGraphInstance> FAnimGraphReference::GetOrAllocate(FExecutionContext& InContext, const FTraitBinding& InBinding) const
+{
+	if (Ptr.IsValid())
+	{
+		Ptr->AllocateVariablesInternal(nullptr, InBinding.GetTraitPtr().GetNodeInstance()->GetOwner().AsShared());
+		Ptr->AllocateGraphInternal();
+		return Ptr;
+	}
+	return TSharedPtr<FAnimNextGraphInstance>();
+}
+
+void FAnimGraphReference::QueueTask(FUniqueAnimGraphTask&& InTaskFunction) const
+{
+	if (Ptr.IsValid() && Ptr->GetModuleInstance())
+	{
+		Ptr->GetModuleInstance()->QueueTask([TaskFunction = MoveTemp(InTaskFunction), WeakAnimGraph = TWeakPtr<FAnimNextGraphInstance>(Ptr)](const UE::UAF::FModuleTaskContext& InContext) mutable
+		{
+			if (TSharedPtr<FAnimNextGraphInstance> PinnedAnimGraph = WeakAnimGraph.Pin())
+			{
+				PinnedAnimGraph->RunOrQueueTask(MoveTemp(TaskFunction));
+			}
+		});
+	}
+}
+
+void FAnimGraphReference::SetDebugName(FName InName) const
+{
+	if (Ptr.IsValid())
+	{
+		Ptr->SetDebugName(InName);
+	}
+}
+
+void FAnimGraphReference::AddReferencedObjects(FReferenceCollector& InCollector) const
+{
+	if (Ptr.IsValid())
+	{
+		InCollector.AddPropertyReferencesWithStructARO(FAnimNextGraphInstance::StaticStruct(), Ptr.Get());
+	}
+}
+
+}
diff --git a/Engine/Plugins/Experimental/UAF/UAFAnimGraph/Source/UAFAnimGraph/Private/Graph/AnimGraphTaskContext.cpp b/Engine/Plugins/Experimental/UAF/UAFAnimGraph/Source/UAFAnimGraph/Private/Graph/AnimGraphTaskContext.cpp
new file mode 100644
--- /dev/null
+++ b/Engine/Plugins/Experimental/UAF/UAFAnimGraph/Source/UAFAnimGraph/Private/Graph/AnimGraphTaskContext.cpp
@@ -0,0 +1,19 @@
+﻿// Copyright Epic Games, Inc. All Rights Reserved.
+
+#include "Graph/AnimGraphTaskContext.h"
+#include "Graph/AnimNextGraphInstance.h"
+
+namespace UE::UAF
+{
+
+FAnimNextGraphInstance& FAnimGraphTaskContext::GetAnimGraphInstance() const
+{
+	return static_cast<FAnimNextGraphInstance&>(Instance);
+}
+
+FAnimGraphTaskContext::FAnimGraphTaskContext(FAnimNextGraphInstance& InGraphInstance)
+	: FInstanceTaskContext(InGraphInstance)
+{
+}
+
+}
\ No newline at end of file
diff --git a/Engine/Plugins/Experimental/UAF/UAFAnimGraph/Source/UAFAnimGraph/Private/Graph/AnimNextAnimationGraph.cpp b/Engine/Plugins/Experimental/UAF/UAFAnimGraph/Source/UAFAnimGraph/Private/Graph/AnimNextAnimationGraph.cpp
--- a/Engine/Plugins/Experimental/UAF/UAFAnimGraph/Source/UAFAnimGraph/Private/Graph/AnimNextAnimationGraph.cpp
+++ b/Engine/Plugins/Experimental/UAF/UAFAnimGraph/Source/UAFAnimGraph/Private/Graph/AnimNextAnimationGraph.cpp
@@ -19,6 +19,8 @@
 #include "DataRegistryTypes.h"
 #include "Module/AnimNextSkeletalMeshComponentReferenceComponent.h"
 #include "RigVMBlueprintGeneratedClass.h"
+#include "UAFAnimGraphInitializerComponent.h"
+#include "Factory/AnimNextFactoryParams.h"
 
 #if WITH_EDITOR
 #include "Editor.h"
@@ -43,7 +45,7 @@ UUAFAnimGraph::UUAFAnimGraph(const FObjectInitializer& ObjectInitializer)
 #endif
 }
 
-TSharedPtr<FAnimNextGraphInstance> UUAFAnimGraph::AllocateInstance(const UE::UAF::FGraphAllocationParams& InParams) const
+TSharedPtr<FAnimNextGraphInstance> UUAFAnimGraph::PreAllocateInstance(const UE::UAF::FGraphAllocationParams& InParams, FAnimNextFactoryParams&& InFactoryParams) const
 {
 	LLM_SCOPE_BYNAME(TEXT("UAF/AnimGraph"))
 	SCOPE_CYCLE_COUNTER(STAT_AnimNext_Graph_AllocateInstance);
@@ -54,10 +56,10 @@ TSharedPtr<FAnimNextGraphInstance> UUAFAnimGraph::AllocateInstance(const UE::UAF
 	const FAnimNextTraitHandle ResolvedRootTraitHandle = ResolvedRootTraitHandles.FindRef(EntryPoint);
 	if (!ResolvedRootTraitHandle.IsValid())
 	{
-		UE_LOG(LogAnimation, Error, TEXT("UUAFAnimGraph::AllocateInstance: Could not allocate graph instance - Root trait handle has not been resolved sucessfully."));
+		UE_LOG(LogAnimation, Error, TEXT("UUAFAnimGraph::AllocateInstance: Could not allocate graph instance - Root trait handle has not been resolved successfully."));
 		return nullptr;
 	}
-	
+
 #if WITH_EDITOR
 	if (CompilationState == EAnimNextRigVMAssetState::CompiledWithErrors)
 	{
@@ -67,42 +69,30 @@ TSharedPtr<FAnimNextGraphInstance> UUAFAnimGraph::AllocateInstance(const UE::UAF
 
 	TSharedPtr<FAnimNextGraphInstance> InstanceImpl = MakeShared<FAnimNextGraphInstance>();
 	InstanceImpl->Asset = this;
-	InstanceImpl->ModuleInstance = InParams.ModuleInstance;
-	InstanceImpl->HostInstance = InParams.ParentGraphInstance != nullptr ? static_cast<FUAFAssetInstance*>(InParams.ParentGraphInstance) : static_cast<FUAFAssetInstance*>(InParams.ModuleInstance);
+	InstanceImpl->ModuleInstance = InParams.SystemReference.WeakPtr;
 	InstanceImpl->EntryPoint = EntryPoint;
-
-	// If we have a parent graph, use its root since we share the same root, otherwise if we have no parent, we are the root
-	InstanceImpl->RootGraphInstance = InParams.ParentGraphInstance != nullptr ? InParams.ParentGraphInstance->GetRootGraphInstance() : InstanceImpl.Get();
-
-	// Copy any built-in components from the asset
-	InstanceImpl->CopyDefaultComponents();
-
-	// Set up variable's data
-	InstanceImpl->InitializeVariables(InParams.Overrides);
-
-	// Bind components now variables are set up
-	InstanceImpl->BindDefaultComponents();
-
+	if (InFactoryParams.IsValid())
 	{
-		UE::UAF::FExecutionContext Context(*InstanceImpl.Get());
-
-		if (InParams.ParentContext != nullptr)
-		{
-			Context.SetBindingObject(InParams.ParentContext->GetBindingObject());
-		}
-		else if (InParams.ModuleInstance != nullptr)
-		{
-			const FAnimNextSkeletalMeshComponentReferenceComponent& ComponentReference = InParams.ModuleInstance->GetComponent<FAnimNextSkeletalMeshComponentReferenceComponent>();
+		InstanceImpl->Components.Add(TInstancedStruct<FUAFAnimGraphInitializerComponent>::Make(MoveTemp(InFactoryParams)));
+	}
+	InstanceImpl->AllocationState = FAnimNextGraphInstance::EAllocationState::PreAllocated;
+	return InstanceImpl;
+}
 
-			if (USkeletalMeshComponent* Component = ComponentReference.GetComponent())
-			{
-				Context.SetBindingObject(Component);
-			}
-		}
+TSharedPtr<FAnimNextGraphInstance> UUAFAnimGraph::AllocateInstance(const UE::UAF::FGraphAllocationParams& InParams, FAnimNextFactoryParams&& InFactoryParams) const
+{
+	LLM_SCOPE_BYNAME(TEXT("UAF/AnimGraph"))
+	SCOPE_CYCLE_COUNTER(STAT_AnimNext_Graph_AllocateInstance);
 
-		InstanceImpl->GraphInstancePtr = Context.AllocateNodeInstance(*InstanceImpl, ResolvedRootTraitHandle);
+	TSharedPtr<FAnimNextGraphInstance> InstanceImpl = PreAllocateInstance(InParams, MoveTemp(InFactoryParams));
+	if (!InstanceImpl.IsValid())
+	{
+		return InstanceImpl;
 	}
 
+	InstanceImpl->AllocateVariablesInternal(InParams.Overrides, InParams.ParentGraphInstance);
+	InstanceImpl->AllocateGraphInternal();
+
 	if (!InstanceImpl->IsValid())
 	{
 		// We failed to allocate our instance, reset the ptr
diff --git a/Engine/Plugins/Experimental/UAF/UAFAnimGraph/Source/UAFAnimGraph/Private/Graph/AnimNextGraphInstance.cpp b/Engine/Plugins/Experimental/UAF/UAFAnimGraph/Source/UAFAnimGraph/Private/Graph/AnimNextGraphInstance.cpp
--- a/Engine/Plugins/Experimental/UAF/UAFAnimGraph/Source/UAFAnimGraph/Private/Graph/AnimNextGraphInstance.cpp
+++ b/Engine/Plugins/Experimental/UAF/UAFAnimGraph/Source/UAFAnimGraph/Private/Graph/AnimNextGraphInstance.cpp
@@ -1,9 +1,13 @@
 // Copyright Epic Games, Inc. All Rights Reserved.
 
+#include "Graph/AnimNextGraphInstance.h"
 #include "Graph/AnimNextGraphInstance.h"
 
 #include "AnimNextAnimGraphStats.h"
 #include "ObjectTrace.h"
+#include "UAFAnimGraphInitializerComponent.h"
+#include "Factory/AnimNextFactoryParams.h"
+#include "Graph/AnimGraphTaskContext.h"
 #include "Script/UAFRigVMComponent.h"
 #include "Graph/AnimNextAnimationGraph.h"
 #include "Graph/AnimNextGraphContextData.h"
@@ -14,6 +18,7 @@
 #include "Logging/StructuredLog.h"
 #include "Module/AnimNextModule.h"
 #include "Module/AnimNextModuleInstance.h"
+#include "Module/AnimNextSkeletalMeshComponentReferenceComponent.h"
 #include "TraitInterfaces/IGarbageCollection.h"
 
 #include UE_INLINE_GENERATED_CPP_BY_NAME(AnimNextGraphInstance)
@@ -27,7 +32,7 @@ FAnimNextGraphInstance::FAnimNextGraphInstance()
 
 FAnimNextGraphInstance::~FAnimNextGraphInstance()
 {
-	TRACE_INSTANCE_LIFETIME_END(ModuleInstance ? ModuleInstance->GetObject() : nullptr, GetUniqueId());
+	TRACE_INSTANCE_LIFETIME_END(ModuleInstance.IsValid() ? ModuleInstance.Pin()->GetObject() : nullptr, GetUniqueId());
 	Release();
 }
 
@@ -47,7 +52,7 @@ void FAnimNextGraphInstance::Release()
 	}
 
 	GraphInstancePtr.Reset();
-	ModuleInstance = nullptr;
+	ModuleInstance.Reset();
 	RootGraphInstance = nullptr;
 	GCHandler.TraitsWithReferences.Reset();
 	ReleaseComponents();
@@ -76,7 +81,7 @@ UE::UAF::FWeakTraitPtr FAnimNextGraphInstance::GetGraphRootPtr() const
 
 FAnimNextModuleInstance* FAnimNextGraphInstance::GetModuleInstance() const
 {
-	return ModuleInstance;
+	return ModuleInstance.Pin().Get();
 }
 
 FAnimNextGraphInstance* FAnimNextGraphInstance::GetParentGraphInstance() const
@@ -87,7 +92,7 @@ FAnimNextGraphInstance* FAnimNextGraphInstance::GetParentGraphInstance() const
 	}
 	else
 	{
-		return static_cast<FAnimNextGraphInstance*>(HostInstance);
+		return static_cast<FAnimNextGraphInstance*>(HostInstance.Pin().Get());
 	}
 }
 
@@ -186,6 +191,88 @@ void FAnimNextGraphInstance::CopyVariablesToLatentPins(const TConstArrayView<UE:
 	}
 }
 
+void FAnimNextGraphInstance::AllocateVariablesInternal(const TSharedPtr<const UE::UAF::FVariableOverridesCollection>& InOverrides, TWeakPtr<FAnimNextGraphInstance> InParentGraphInstance)
+{
+	using namespace UE::UAF;
+	
+	check(AllocationState >= EAllocationState::PreAllocated);
+
+	if (AllocationState >= EAllocationState::VariablesAllocated)
+	{
+		// Nothing to do, already allocated variables
+		return;
+	}
+
+	if (InParentGraphInstance.IsValid())
+	{
+		HostInstance = InParentGraphInstance;
+		RootGraphInstance = InParentGraphInstance.Pin()->GetRootGraphInstance();
+	}
+	else
+	{
+		HostInstance = ModuleInstance;
+		RootGraphInstance = this;
+	}
+
+	auto FindInitializerComponent = [](const TInstancedStruct<FUAFAssetInstanceComponent>& InComponent)
+	{
+		return InComponent.GetScriptStruct() == FUAFAnimGraphInitializerComponent::StaticStruct();
+	};
+
+	TInstancedStruct<FUAFAssetInstanceComponent> InitializerStruct;
+	int32 InitializerComponentIndex = Components.IndexOfByPredicate(FindInitializerComponent);
+	if (InitializerComponentIndex != INDEX_NONE)
+	{
+		InitializerStruct = MoveTemp(Components[InitializerComponentIndex]);
+		Components.RemoveAt(InitializerComponentIndex, EAllowShrinking::No);
+	}
+
+	// Copy any built-in components from the asset
+	CopyDefaultComponents();
+
+	// Set up variable's data
+	InitializeVariables(InOverrides);
+
+	// Bind components now variables are set up
+	BindDefaultComponents();
+
+	// Flush any initialization tasks we have
+	if (InitializerStruct.IsValid())
+	{
+		const FUAFAnimGraphInitializerComponent& InitializerComponent = InitializerStruct.Get<FUAFAnimGraphInitializerComponent>();
+		InitializerComponent.FactoryParams.InitializeInstance(*this);
+		for (const FUniqueAnimGraphTask& Task : InitializerComponent.TaskQueue)
+		{
+			Task(FAnimGraphTaskContext(*this));
+		}
+	}
+
+	AllocationState = EAllocationState::VariablesAllocated;
+}
+
+void FAnimNextGraphInstance::AllocateGraphInternal()
+{
+	check(AllocationState >= EAllocationState::VariablesAllocated);
+
+	if (AllocationState == EAllocationState::Allocated)
+	{
+		// Nothing to do, already allocated
+		return;
+	}
+
+	// Verify the resolved entry point
+	const FAnimNextTraitHandle ResolvedRootTraitHandle = GetAnimationGraph()->ResolvedRootTraitHandles.FindRef(EntryPoint);
+	if (!ensure(ResolvedRootTraitHandle.IsValid()))
+	{
+		return;
+	}
+
+	UE::UAF::FExecutionContext Context(*this);
+	GraphInstancePtr = Context.AllocateNodeInstance(*this, ResolvedRootTraitHandle);
+
+	AllocationState = EAllocationState::Allocated;
+}
+
 #if WITH_EDITORONLY_DATA
 void FAnimNextGraphInstance::Freeze()
 {
@@ -252,6 +339,21 @@ void FAnimNextGraphInstance::UnregisterWithGC(const UE::UAF::FWeakTraitPtr& InTr
 	}
 }
 
+void FAnimNextGraphInstance::RunOrQueueTask(UE::UAF::FUniqueAnimGraphTask&& InTaskFunction)
+{
+	using namespace UE::UAF;
+
+	if (AllocationState == EAllocationState::PreAllocated)
+	{
+		FUAFAnimGraphInitializerComponent& InitializerComponent = GetComponent<FUAFAnimGraphInitializerComponent>();
+		InitializerComponent.TaskQueue.Add(MoveTemp(InTaskFunction));
+	}
+	else
+	{
+		InTaskFunction(UE::UAF::FAnimGraphTaskContext(*this));
+	}
+}
+
 void FAnimNextGraphInstance::AddStructReferencedObjects(FReferenceCollector& Collector)
 {
 	GCHandler.AddReferencedObjects(Collector);
diff --git a/Engine/Plugins/Experimental/UAF/UAFAnimGraph/Source/UAFAnimGraph/Private/Graph/AnimNextModuleAnimGraphComponent.cpp b/Engine/Plugins/Experimental/UAF/UAFAnimGraph/Source/UAFAnimGraph/Private/Graph/AnimNextModuleAnimGraphComponent.cpp
--- a/Engine/Plugins/Experimental/UAF/UAFAnimGraph/Source/UAFAnimGraph/Private/Graph/AnimNextModuleAnimGraphComponent.cpp
+++ b/Engine/Plugins/Experimental/UAF/UAFAnimGraph/Source/UAFAnimGraph/Private/Graph/AnimNextModuleAnimGraphComponent.cpp
@@ -14,7 +14,7 @@ TWeakPtr<FAnimNextGraphInstance> FAnimNextModuleAnimGraphComponent::AllocateInst
 
 	TSharedPtr<FAnimNextGraphInstance> NewInstance = InAnimationGraph->AllocateInstance(
 		{
-			.ModuleInstance = &GetModuleInstance(),
+			.SystemReference = GetModuleInstance().GetReference(),
 			.Overrides = InOverrides
 		});
 	if(!NewInstance.IsValid())
diff --git a/Engine/Plugins/Experimental/UAF/UAFAnimGraph/Source/UAFAnimGraph/Private/Graph/InputPoseImpl_ExternalGraph.cpp b/Engine/Plugins/Experimental/UAF/UAFAnimGraph/Source/UAFAnimGraph/Private/Graph/InputPoseImpl_ExternalGraph.cpp
new file mode 100644
--- /dev/null
+++ b/Engine/Plugins/Experimental/UAF/UAFAnimGraph/Source/UAFAnimGraph/Private/Graph/InputPoseImpl_ExternalGraph.cpp
@@ -0,0 +1,23 @@
+﻿// Copyright Epic Games, Inc. All Rights Reserved.
+
+#include "Graph/InputPoseImpl_ExternalGraph.h"
+
+namespace UE::UAF
+{
+
+const FAnimNextGraphLODPose* FValueBundleImpl_ExternalGraph::GetLODPose() const
+{
+	return nullptr;
+}
+
+const FValueBundleHeap* FValueBundleImpl_ExternalGraph::GetValueBundle() const
+{
+	return nullptr;
+}
+
+const FWeakAnimGraphReference* FValueBundleImpl_ExternalGraph::GetAnimGraphReference() const
+{
+	return &WeakGraphReference;
+}
+
+}
diff --git a/Engine/Plugins/Experimental/UAF/UAFAnimGraph/Source/UAFAnimGraph/Private/Graph/InputPoseImpl_ValueBundle.cpp b/Engine/Plugins/Experimental/UAF/UAFAnimGraph/Source/UAFAnimGraph/Private/Graph/InputPoseImpl_ValueBundle.cpp
new file mode 100644
--- /dev/null
+++ b/Engine/Plugins/Experimental/UAF/UAFAnimGraph/Source/UAFAnimGraph/Private/Graph/InputPoseImpl_ValueBundle.cpp
@@ -0,0 +1,21 @@
+﻿// Copyright Epic Games, Inc. All Rights Reserved.
+
+#include "Graph/InputPoseImpl_ValueBundle.h"
+
+#include "EvaluationVM/Tasks/PushPose.h"
+#include "TraitInterfaces/IEvaluate.h"
+
+namespace UE::UAF
+{
+
+const FAnimNextGraphLODPose* FValueBundleImpl_ValueBundle::GetLODPose() const
+{
+	return &Pose;
+}
+
+const FValueBundleHeap* FValueBundleImpl_ValueBundle::GetValueBundle() const
+{
+	return &ValueBundle;
+}
+
+}
diff --git a/Engine/Plugins/Experimental/UAF/UAFAnimGraph/Source/UAFAnimGraph/Private/Graph/InputPoseImpl_ValueBundle.h b/Engine/Plugins/Experimental/UAF/UAFAnimGraph/Source/UAFAnimGraph/Private/Graph/InputPoseImpl_ValueBundle.h
new file mode 100644
--- /dev/null
+++ b/Engine/Plugins/Experimental/UAF/UAFAnimGraph/Source/UAFAnimGraph/Private/Graph/InputPoseImpl_ValueBundle.h
@@ -0,0 +1,45 @@
+﻿// Copyright Epic Games, Inc. All Rights Reserved.
+
+#pragma once
+
+#include "Graph/AnimNext_LODPose.h"
+#include "Graph/WeakAnimGraphReference.h"
+#include "UAF/ValueRuntime/IVirtualValueBundle.h"
+#include "UAF/ValueRuntime/ValueBundle.h"
+
+#define UE_API UAFANIMGRAPH_API
+
+namespace UE::UAF
+{
+
+struct FValueBundleImpl_ValueBundle : public IVirtualValueBundle
+{
+	explicit FValueBundleImpl_ValueBundle(const FAnimNextGraphLODPose& InPose)
+		: Pose(InPose)
+	{}
+
+	explicit FValueBundleImpl_ValueBundle(FAnimNextGraphLODPose&& InPose)
+		: Pose(MoveTemp(InPose))
+	{}
+
+	explicit FValueBundleImpl_ValueBundle(const FValueBundleHeap& InValueBundle)
+		: ValueBundle(InValueBundle)
+	{}
+	
+		explicit FValueBundleImpl_ValueBundle(FValueBundleHeap&& InValueBundle)
+		: ValueBundle(MoveTemp(InValueBundle))
+	{}
+	
+	// IVirtualValueBundle interface
+	virtual const FAnimNextGraphLODPose* GetLODPose() const override;
+	virtual const FValueBundleHeap* GetValueBundle() const override;
+	virtual const FWeakAnimGraphReference* GetAnimGraphReference() const override { return nullptr; }
+
+private:
+	FAnimNextGraphLODPose Pose;
+	FValueBundleHeap ValueBundle;
+};
+
+}
+
+#undef UE_API
\ No newline at end of file
diff --git a/Engine/Plugins/Experimental/UAF/UAFAnimGraph/Source/UAFAnimGraph/Private/Graph/RigUnit_AnimNextRunAnimationGraph_v1.cpp b/Engine/Plugins/Experimental/UAF/UAFAnimGraph/Source/UAFAnimGraph/Private/Graph/RigUnit_AnimNextRunAnimationGraph_v1.cpp
--- a/Engine/Plugins/Experimental/UAF/UAFAnimGraph/Source/UAFAnimGraph/Private/Graph/RigUnit_AnimNextRunAnimationGraph_v1.cpp
+++ b/Engine/Plugins/Experimental/UAF/UAFAnimGraph/Source/UAFAnimGraph/Private/Graph/RigUnit_AnimNextRunAnimationGraph_v1.cpp
@@ -76,7 +76,6 @@ FRigUnit_AnimNextRunAnimationGraph_v1_Execute()
 
 	// Every graph in a schedule will see the same input events (if they were queued before the schedule started)
 	FUpdateGraphContext UpdateGraphContext(InstanceRef.Get(), ExecuteContext.GetDeltaTime());
-	UpdateGraphContext.SetBindingObject(RefPose.SkeletalMeshComponent);
 	FTraitEventList& InputEventList = UpdateGraphContext.GetInputEventList();
 
 	// A schedule can contain multiple graphs, we copy the input event list since it might be appended to during our update
@@ -99,7 +98,6 @@ FRigUnit_AnimNextRunAnimationGraph_v1_Execute()
 	UpdateGraph(UpdateGraphContext);
 	
 	FEvaluateGraphContext EvaluateGraphContext(InstanceRef.Get(), RefPose, DesiredLOD);
-	EvaluateGraphContext.SetBindingObject(RefPose.SkeletalMeshComponent);
 
 	if (UseAbstractSkeletonRuntime())
 	{
diff --git a/Engine/Plugins/Experimental/UAF/UAFAnimGraph/Source/UAFAnimGraph/Private/Graph/RigUnit_AnimNextRunAnimationGraph_v2.cpp b/Engine/Plugins/Experimental/UAF/UAFAnimGraph/Source/UAFAnimGraph/Private/Graph/RigUnit_AnimNextRunAnimationGraph_v2.cpp
--- a/Engine/Plugins/Experimental/UAF/UAFAnimGraph/Source/UAFAnimGraph/Private/Graph/RigUnit_AnimNextRunAnimationGraph_v2.cpp
+++ b/Engine/Plugins/Experimental/UAF/UAFAnimGraph/Source/UAFAnimGraph/Private/Graph/RigUnit_AnimNextRunAnimationGraph_v2.cpp
@@ -16,6 +16,8 @@
 #include "Factory/AnimGraphFactory.h"
 #include "UAF/UAFSkeletonUserData.h"
 #include "UAFLogging.h"
+#include "Graph/RigUnit_UAFRunAsset.h"
+#include "Graph/UAFSystemOutputComponent.h"
 
 #include UE_INLINE_GENERATED_CPP_BY_NAME(RigUnit_AnimNextRunAnimationGraph_v2)
 
@@ -28,14 +30,23 @@ FRigUnit_AnimNextRunAnimationGraph_v2_Execute()
 	const FAnimNextModuleContextData& ModuleContextData = ExecuteContext.GetContextData<FAnimNextModuleContextData>();
 	FAnimNextModuleInstance& ModuleInstance = ModuleContextData.GetModuleInstance();
 
-	if (!ReferencePose.ReferencePose.IsValid())
+	FDataHandle RefPoseDataHandle = ReferencePose.ReferencePose;
+
+	// Fall back to output pose ref pose
+	if (!RefPoseDataHandle.IsValid())
+	{
+		FUAFSystemOutputComponent& OutputPoseComponent = ModuleInstance.GetComponent<FUAFSystemOutputComponent>();
+		RefPoseDataHandle = OutputPoseComponent.GetRefPose().ReferencePose;
+	}
+
+	if (!RefPoseDataHandle.IsValid())
 	{
 		UAF_RIGUNIT_LOG(Warning, TEXT("Could not run graph - Ref Pose is invalid. [Asset: %s] [HostGraph: %s]"), 
 			*GetNameSafe(Graph.Asset), *GetNameSafe(Graph.HostGraph));
 		return;
 	}
 
-	const UE::UAF::FReferencePose& RefPose = ReferencePose.ReferencePose.GetRef<UE::UAF::FReferencePose>();
+	const UE::UAF::FReferencePose& RefPose = RefPoseDataHandle.GetRef<UE::UAF::FReferencePose>();
 
 	int32 DesiredLOD = LOD;
 	if (DesiredLOD == -1)
@@ -104,7 +115,6 @@ FRigUnit_AnimNextRunAnimationGraph_v2_Execute()
 
 	// Every graph in a schedule will see the same input events (if they were queued before the schedule started)
 	FUpdateGraphContext UpdateGraphContext(HostInstance, ExecuteContext.GetDeltaTime());
-	UpdateGraphContext.SetBindingObject(RefPose.SkeletalMeshComponent);
 
 	FTraitEventList& InputEventList = UpdateGraphContext.GetInputEventList();
 
@@ -128,7 +138,6 @@ FRigUnit_AnimNextRunAnimationGraph_v2_Execute()
 	UpdateGraph(UpdateGraphContext);
 
 	FEvaluateGraphContext EvaluateGraphContext(HostInstance, RefPose, DesiredLOD);
-	EvaluateGraphContext.SetBindingObject(RefPose.SkeletalMeshComponent);
 
 	if (UseAbstractSkeletonRuntime())
 	{
@@ -166,3 +175,14 @@ FRigUnit_AnimNextRunAnimationGraph_v2_Execute()
 		ModuleInstance.OutputEventList.Append(UpdateGraphContext.GetOutputEventList());
 	}
 }
+
+FRigVMStructUpgradeInfo FRigUnit_AnimNextRunAnimationGraph_v2::GetUpgradeInfo() const
+{
+	FRigUnit_UAFRunAsset NewNode;
+	NewNode.Graph = Graph;
+	NewNode.Overrides = Overrides;
+	NewNode.AttributeSet = AttributeSet;
+
+	FRigVMStructUpgradeInfo Info(*this, NewNode);
+	return Info;
+}
\ No newline at end of file
diff --git a/Engine/Plugins/Experimental/UAF/UAFAnimGraph/Source/UAFAnimGraph/Private/Graph/RigUnit_UAFRunAsset.cpp b/Engine/Plugins/Experimental/UAF/UAFAnimGraph/Source/UAFAnimGraph/Private/Graph/RigUnit_UAFRunAsset.cpp
new file mode 100644
--- /dev/null
+++ b/Engine/Plugins/Experimental/UAF/UAFAnimGraph/Source/UAFAnimGraph/Private/Graph/RigUnit_UAFRunAsset.cpp
@@ -0,0 +1,165 @@
+// Copyright Epic Games, Inc. All Rights Reserved.
+
+#include "Graph/RigUnit_UAFRunAsset.h"
+
+#include "InputPoseImpl_ValueBundle.h"
+#include "TraitInterfaces/IEvaluate.h"
+#include "TraitInterfaces/IUpdate.h"
+#include "Module/AnimNextModuleInstance.h"
+#include "Graph/AnimNextAnimationGraph.h"
+#include "Graph/AnimNextGraphInstance.h"
+#include "Graph/AnimNextModuleAnimGraphComponent.h"
+#include "Graph/RigVMTrait_AnimNextPublicVariables.h"
+#include "Injection/GraphInstanceInjectionComponent.h"
+#include "TraitCore/TraitEventList.h"
+#include "Injection/AnimNextModuleInjectionComponent.h"
+#include "Graph/UAFAnimGraphTrace.h"
+#include "Factory/AnimGraphFactory.h"
+#include "UAF/UAFSkeletonUserData.h"
+#include "UAFLogging.h"
+#include "Graph/UAFSystemOutputComponent.h"
+
+#include UE_INLINE_GENERATED_CPP_BY_NAME(RigUnit_UAFRunAsset)
+
+FRigUnit_UAFRunAsset_Execute()
+{
+	SCOPED_NAMED_EVENT(UAF_Run_Asset, FColor::Orange);
+
+	using namespace UE::UAF;
+
+	const FAnimNextModuleContextData& ModuleContextData = ExecuteContext.GetContextData<FAnimNextModuleContextData>();
+	FAnimNextModuleInstance& ModuleInstance = ModuleContextData.GetModuleInstance();
+
+	FUAFSystemOutputComponent& OutputPoseComponent = ModuleInstance.GetComponent<FUAFSystemOutputComponent>();
+	FDataHandle RefPoseDataHandle = OutputPoseComponent.GetRefPose().ReferencePose;
+	if (!RefPoseDataHandle.IsValid())
+	{
+		UAF_RIGUNIT_LOG(Warning, TEXT("Could not run graph - Ref Pose is invalid. [Asset: %s] [HostGraph: %s]"), 
+			*GetNameSafe(Graph.Asset), *GetNameSafe(Graph.HostGraph));
+		return;
+	}
+
+	const UE::UAF::FReferencePose& RefPose = RefPoseDataHandle.GetRef<UE::UAF::FReferencePose>();
+
+	const int32 DesiredLOD = RefPose.GetSourceLODLevel();
+	FAnimNextGraphLODPose GraphPose;
+	// TODO: Currently forcing additive flag to false here
+	GraphPose.LODPose.PrepareForLOD(RefPose, DesiredLOD, true, false);
+
+	ensure(GraphPose.LODPose.LODLevel == DesiredLOD);
+
+	// Get a host to run this graph
+	const UUAFAnimGraph* HostGraph = Graph.HostGraph ? Graph.HostGraph.Get() : FAnimGraphFactory::GetDefaultGraphHost();
+	FAnimNextModuleAnimGraphComponent& AnimationGraphComponent = ModuleInstance.GetComponent<FAnimNextModuleAnimGraphComponent>();
+	if(HostGraph == nullptr)
+	{
+		AnimationGraphComponent.ReleaseInstance(WorkData.WeakHost);
+		UAF_RIGUNIT_LOG(Warning, TEXT("Could not run graph - No valid host graph could be found. [Asset: %s] [HostGraph: %s]"),
+			*GetNameSafe(Graph.Asset), *GetNameSafe(Graph.HostGraph));
+		return;
+	}
+
+	// Release the instance if the host graph has changed
+	if(WorkData.WeakHost.IsValid() && !WorkData.WeakHost.Pin()->UsesAnimationGraph(HostGraph))
+	{
+		AnimationGraphComponent.ReleaseInstance(WorkData.WeakHost);
+	}
+
+	// Lazily (re-)allocate graph instance if required
+	if(!WorkData.WeakHost.IsValid())
+	{
+		// Reset our graph reference, we may be allocating a new graph
+		WorkData.InjectedGraphReference.Reset();
+
+		TSharedPtr<FAnimNextGraphInstance> NewGraphInstance = AnimationGraphComponent.AllocateInstance(HostGraph, Overrides.GetOverrides().Pin()).Pin();
+		if (NewGraphInstance.IsValid())
+		{
+			WorkData.WeakHost = NewGraphInstance;
+
+			FGraphInstanceInjectionComponent& HostInjectionComponent = NewGraphInstance->GetComponent<FGraphInstanceInjectionComponent>();
+			WorkData.InjectedGraphReference = HostInjectionComponent.GetInjectionInfo().GetDefaultInjectionSite();
+		}
+	}
+
+	if(!WorkData.WeakHost.IsValid())
+	{
+		UAF_RIGUNIT_LOG(Warning, TEXT("Could not run graph - Host Graph could not be allocated. [Asset: %s] [HostGraph: %s]"),
+			*GetNameSafe(Graph.Asset), *GetNameSafe(Graph.HostGraph));
+		return;
+	}
+
+	// Take a strong reference to the host instance, we are going to run it
+	TSharedRef<FAnimNextGraphInstance> HostInstanceRef = WorkData.WeakHost.Pin().ToSharedRef();
+	FAnimNextGraphInstance& HostInstance = HostInstanceRef.Get();
+
+	// Re-set the injected graph each time, as it may change per-update
+	if (!WorkData.InjectedGraphReference.IsNone())
+	{
+		HostInstance.SetVariable(WorkData.InjectedGraphReference, Graph);
+	}
+
+	// Every graph in a schedule will see the same input events (if they were queued before the schedule started)
+	FUpdateGraphContext UpdateGraphContext(HostInstance, ExecuteContext.GetDeltaTime());
+
+	FTraitEventList& InputEventList = UpdateGraphContext.GetInputEventList();
+
+	// A module can contain multiple graphs, we copy the input event list since it might be appended to during our update
+	{
+		UE::TReadScopeLock ReadLock(ModuleInstance.EventListLock);
+		InputEventList = ModuleInstance.InputEventList;
+	}
+
+	// Track how many input events we started with, we'll append the new ones
+	const int32 NumOriginalInputEvents = InputEventList.Num();
+
+	// Internally we use memstack allocation, so we need a mark here
+	FMemStack& MemStack = FMemStack::Get();
+	FMemMark MemMark(MemStack);
+
+	// We allocate a dummy buffer to trigger the allocation of a large chunk if this is the first mark
+	// This reduces churn internally by avoiding a chunk to be repeatedly allocated and freed as we push/pop marks
+	MemStack.Alloc(size_t(FPageAllocator::SmallPageSize) + 1, 16);
+
+	UpdateGraph(UpdateGraphContext);
+
+	FEvaluateGraphContext EvaluateGraphContext(HostInstance, RefPose, DesiredLOD);
+
+	if (UseAbstractSkeletonRuntime())
+	{
+		// By default, use the everything set
+		USkeleton* Skeleton = const_cast<USkeleton*>(RefPose.GetSkeletonAsset());
+		UUAFSkeletonUserData* UserData = UUAFSkeletonUserData::FromSkeleton(*Skeleton);
+
+		UAbstractSkeletonSetBinding* Binding = UserData != nullptr ? UserData->GetSetBinding() : nullptr;
+
+		EvaluateGraphContext.SetSkeletonSet(Binding, AttributeSet);
+	}
+
+	EvaluateGraph(EvaluateGraphContext, GraphPose);
+
+	TRACE_UAF_MODULE(ModuleInstance);
+	TRACE_UAF_GRAPHINSTANCES(HostInstance);
+
+	// We might have appended new input/output events, append them
+	{
+		const int32 NumInputEvents = InputEventList.Num();
+
+		UE::TWriteScopeLock WriteLock(ModuleInstance.EventListLock);
+
+		// Append the new input events
+		for (int32 EventIndex = NumOriginalInputEvents; EventIndex < NumInputEvents; ++EventIndex)
+		{
+			FAnimNextTraitEventPtr& Event = InputEventList[EventIndex];
+			if (Event->IsValid())
+			{
+				ModuleInstance.InputEventList.Push(MoveTemp(Event));
+			}
+		}
+
+		// Append our output events
+		ModuleInstance.OutputEventList.Append(UpdateGraphContext.GetOutputEventList());
+	}
+
+	// Move our output pose into the result
+	Result.SetAs<FValueBundleImpl_ValueBundle>(MoveTemp(GraphPose));
+}
diff --git a/Engine/Plugins/Experimental/UAF/UAFAnimGraph/Source/UAFAnimGraph/Private/Graph/UAFAnimGraphInitializerComponent.h b/Engine/Plugins/Experimental/UAF/UAFAnimGraph/Source/UAFAnimGraph/Private/Graph/UAFAnimGraphInitializerComponent.h
new file mode 100644
--- /dev/null
+++ b/Engine/Plugins/Experimental/UAF/UAFAnimGraph/Source/UAFAnimGraph/Private/Graph/UAFAnimGraphInitializerComponent.h
@@ -0,0 +1,43 @@
+﻿// Copyright Epic Games, Inc. All Rights Reserved.
+
+#pragma once
+
+#include "CoreMinimal.h"
+#include "Graph/UAFGraphInstanceComponent.h"
+#include "Factory/AnimNextFactoryParams.h"
+#include "Graph/AnimGraphTaskContext.h"
+#include "UAFAnimGraphInitializerComponent.generated.h"
+
+namespace UE::UAF
+{
+struct FExecutionContext;
+}
+
+/** Internal component used for initial UAF graph setup */
+USTRUCT(meta=(Hidden, Abstract))
+struct FUAFAnimGraphInitializerComponent : public FUAFGraphInstanceComponent
+{
+	GENERATED_BODY()
+	DECLARE_UAF_ASSET_INSTANCE_COMPONENT()
+
+	FUAFAnimGraphInitializerComponent() = default;
+	explicit FUAFAnimGraphInitializerComponent(FAnimNextFactoryParams&& InFactoryParams)
+		: FactoryParams(MoveTemp(InFactoryParams))
+	{}
+
+	// Factory params used to construct the graph instance
+	FAnimNextFactoryParams FactoryParams;
+
+	// Tasks queue to be applied on construction
+	TArray<UE::UAF::FUniqueAnimGraphTask> TaskQueue;
+};
+
+template<>
+struct TStructOpsTypeTraits<FUAFAnimGraphInitializerComponent> : public TStructOpsTypeTraitsBase2<FUAFAnimGraphInitializerComponent>
+{
+	enum
+	{
+		WithCopy = false
+	};
+};
+
diff --git a/Engine/Plugins/Experimental/UAF/UAFAnimGraph/Source/UAFAnimGraph/Private/Graph/WeakAnimGraphReference.cpp b/Engine/Plugins/Experimental/UAF/UAFAnimGraph/Source/UAFAnimGraph/Private/Graph/WeakAnimGraphReference.cpp
new file mode 100644
--- /dev/null
+++ b/Engine/Plugins/Experimental/UAF/UAFAnimGraph/Source/UAFAnimGraph/Private/Graph/WeakAnimGraphReference.cpp
@@ -0,0 +1,88 @@
+﻿// Copyright Epic Games, Inc. All Rights Reserved.
+
+#include "Graph/WeakAnimGraphReference.h"
+
+#include "Graph/AnimGraphReference.h"
+#include "Graph/AnimNextAnimationGraph.h"
+#include "Graph/AnimNextGraphInstance.h"
+#include "Module/AnimNextModuleInstance.h"
+#include "TraitCore/ExecutionContext.h"
+#include "TraitCore/NodeInstance.h"
+
+namespace UE::UAF
+{
+
+FWeakAnimGraphReference::FWeakAnimGraphReference(const FAnimGraphReference& InAnimGraphReference)
+{
+	Ptr = InAnimGraphReference.Ptr;
+}
+
+FWeakAnimGraphReference::FWeakAnimGraphReference(const TSharedPtr<FAnimNextGraphInstance>& InAnimGraphInstance)
+{
+	Ptr = InAnimGraphInstance;
+}
+
+void FWeakAnimGraphReference::Reset()
+{
+	Ptr.Reset();
+}
+
+bool FWeakAnimGraphReference::IsValid() const
+{
+	return Ptr.IsValid();
+}
+
+const UUAFAnimGraph* FWeakAnimGraphReference::GetAnimGraph() const
+{
+	if (TSharedPtr<FAnimNextGraphInstance> PinnedPtr = Ptr.Pin())
+	{
+		return PinnedPtr->GetAnimationGraph();
+	}
+	return nullptr;
+}
+
+TSharedPtr<FAnimNextGraphInstance> FWeakAnimGraphReference::GetOrAllocate(const FExecutionContext& InContext, const FTraitBinding& InBinding) const
+{
+	if (TSharedPtr<FAnimNextGraphInstance> PinnedPtr = Ptr.Pin())
+	{
+		PinnedPtr->AllocateVariablesInternal(nullptr, InBinding.GetTraitPtr().GetNodeInstance()->GetOwner().AsShared());
+		PinnedPtr->AllocateGraphInternal();
+		return PinnedPtr;
+	}
+	return TSharedPtr<FAnimNextGraphInstance>();
+}
+
+void FWeakAnimGraphReference::QueueTask(FUniqueAnimGraphTask&& InTaskFunction) const
+{
+	if (TSharedPtr<FAnimNextGraphInstance> PinnedPtr = Ptr.Pin())
+	{
+		if(FAnimNextModuleInstance* ModuleInstance = PinnedPtr->GetModuleInstance())
+		{
+			ModuleInstance->QueueTask([TaskFunction = MoveTemp(InTaskFunction), WeakAnimGraph = Ptr](const UE::UAF::FModuleTaskContext& InContext) mutable
+			{
+				if (TSharedPtr<FAnimNextGraphInstance> PinnedAnimGraph = WeakAnimGraph.Pin())
+				{
+					PinnedAnimGraph->RunOrQueueTask(MoveTemp(TaskFunction));
+				}
+			});
+		}
+	}
+}
+
+void FWeakAnimGraphReference::SetDebugName(FName InName) const
+{
+	if (TSharedPtr<FAnimNextGraphInstance> PinnedPtr = Ptr.Pin())
+	{
+		PinnedPtr->SetDebugName(InName);
+	}
+}
+
+void FWeakAnimGraphReference::AddReferencedObjects(FReferenceCollector& InCollector) const
+{
+	if (TSharedPtr<FAnimNextGraphInstance> PinnedPtr = Ptr.Pin())
+	{
+		InCollector.AddPropertyReferencesWithStructARO(FAnimNextGraphInstance::StaticStruct(), PinnedPtr.Get());
+	}
+}
+
+}
diff --git a/Engine/Plugins/Experimental/UAF/UAFAnimGraph/Source/UAFAnimGraph/Private/Native/UAFSimplePrePhysicsGraphComponent.cpp b/Engine/Plugins/Experimental/UAF/UAFAnimGraph/Source/UAFAnimGraph/Private/Native/UAFSimplePrePhysicsGraphComponent.cpp
--- a/Engine/Plugins/Experimental/UAF/UAFAnimGraph/Source/UAFAnimGraph/Private/Native/UAFSimplePrePhysicsGraphComponent.cpp
+++ b/Engine/Plugins/Experimental/UAF/UAFAnimGraph/Source/UAFAnimGraph/Private/Native/UAFSimplePrePhysicsGraphComponent.cpp
@@ -2,15 +2,14 @@
 
 #include "Native/UAFSimplePrePhysicsGraphComponent.h"
 
-#include "Components/SkeletalMeshComponent.h"
-#include "Graph/RigUnit_AnimNextRunAnimationGraph_v2.h"
+#include "Graph/RigUnit_UAFRunAsset.h"
 #include "Injection/InjectionSiteTrait.h"
 #include "Module/AnimNextModuleInstance.h"
 #include "Module/RigUnit_AnimNextModuleEvents.h"
 #include "Script/UAFScriptContextData.h"
-#include "Graph/RigUnit_AnimNextWriteSkeletalMeshComponentPose.h"
 #include "Module/ModuleTaskContext.h"
 #include "Variables/RigUnit_CopyModuleProxyVariables.h"
+#include "Graph/RigUnit_UAFWriteSystemOutput.h"
 
 TConstArrayView<UE::UAF::FScriptEventInfo> FUAFSimplePrePhysicsGraphComponent::GetScriptEvents()
 {
@@ -27,33 +26,29 @@ void FUAFSimplePrePhysicsGraphComponent::PrePhysics(TStructView<FUAFScriptCompon
 	using namespace UE::UAF;
 	
 	FUAFSimplePrePhysicsGraphComponent& This = InScriptComponent.Get<FUAFSimplePrePhysicsGraphComponent>();
-	if (USkeletalMeshComponent* Mesh = This.WeakMesh.Get())
-	{
-		const FAnimNextModuleContextData& ModuleContextData = InContextData.Get<FAnimNextModuleContextData>();
-		FAnimNextExecuteContext Context;
-		Context.SetDeltaTime(ModuleContextData.GetDeltaTime());
-		Context.SetOwningObject(ModuleContextData.GetModuleInstance().GetObject());
+	
+	const FAnimNextModuleContextData& ModuleContextData = InContextData.Get<FAnimNextModuleContextData>();
+	FAnimNextModuleInstance& ModuleInstance = ModuleContextData.GetModuleInstance();
 
-		FScopedExecuteContextData ScopedExecuteContextData(Context, InContextData);
-		const int32 LOD = -1;
+	FAnimNextExecuteContext Context;
+	Context.SetDeltaTime(ModuleContextData.GetDeltaTime());
+	Context.SetOwningObject(ModuleInstance.GetObject());
 
-		FRigUnit_CopyModuleProxyVariables::StaticExecute(Context);
+	FScopedExecuteContextData ScopedExecuteContextData(Context, InContextData);
 
-		FRigUnit_AnimNextRunAnimationGraph_v2::StaticExecute(
-			Context,
-			*This.Graph,
-			LOD,
-			This.RefPose,
-			FAnimNextVariableOverridesCollection(),
-			NAME_None,
-			This.Pose,
-			This.WorkData);
+	FRigUnit_CopyModuleProxyVariables::StaticExecute(Context);
 
-		FRigUnit_AnimNextWriteSkeletalMeshComponentPose::StaticExecute(
-			Context,
-			This.Pose,
-			Mesh);
-	}
+	FRigUnit_UAFRunAsset::StaticExecute(
+		Context,
+		*This.Graph,
+		FAnimNextVariableOverridesCollection(),
+		NAME_None,
+		This.Output,
+		This.WorkData);
+
+	FRigUnit_UAFWriteSystemOutput::StaticExecute(
+		Context,
+		This.Output);
 }
 
 void FUAFSimplePrePhysicsGraphComponent::CallEventByName(TConstStructView<FUAFScriptContextData> InContextData)
@@ -71,40 +66,11 @@ void FUAFSimplePrePhysicsGraphComponent::OnBindToInstance()
 	// This interacts with actors/components at present, so needs to be called from the GT
 	check(IsInGameThread());
 
-	// TODO: Once we have I/O reworked, we can remove this hack to get the mesh component reference to use
 	FAnimNextModuleInstance& SystemInstance = GetAssetInstance().As<FAnimNextModuleInstance>();
-	UActorComponent* HostComponent = Cast<UActorComponent>(SystemInstance.GetObject());
-	if (!ensure(HostComponent))
-	{
-		return;
-	}
-
-	AActor* Owner = HostComponent->GetOwner();
-	if (!ensure(Owner))
+	SystemInstance.AccessVariablesStruct<FInjectionSiteTraitData>([this](FInjectionSiteTraitData& InjectionSiteTraitData)
 	{
-		return;
-	}
-
-	// Just try to bind to the first mesh component on the actor for this simple case for now
-	USkeletalMeshComponent* Mesh = Owner->FindComponentByClass<USkeletalMeshComponent>();
-	if (Mesh == nullptr)
-	{
-		return;
-	}
-
-	SystemInstance.QueueTask(
-		FRigUnit_AnimNextPrePhysicsEvent::DefaultEventName,
-		[WeakSkeletalMesh = TWeakObjectPtr<USkeletalMeshComponent>(Mesh)](const FModuleTaskContext& InContext)
-		{
-			FUAFSimplePrePhysicsGraphComponent& ThisComponent = InContext.GetModuleInstance()->GetComponent<FUAFSimplePrePhysicsGraphComponent>();
-			ThisComponent.WeakMesh = WeakSkeletalMesh;
-			ThisComponent.RefPose.ReferencePose = UE::UAF::FDataRegistry::Get()->GetOrGenerateReferencePose(WeakSkeletalMesh.Get());
-
-			InContext.GetModuleInstance()->AccessVariablesStruct<FInjectionSiteTraitData>([&ThisComponent](FInjectionSiteTraitData& InjectionSiteTraitData)
-			{
-				ThisComponent.Graph = &InjectionSiteTraitData.Graph;
-			});
+		Graph = &InjectionSiteTraitData.Graph;
+	});
 
-			check(ThisComponent.Graph);
-		});
+	check(Graph);
 }
\ No newline at end of file
diff --git a/Engine/Plugins/Experimental/UAF/UAFAnimGraph/Source/UAFAnimGraph/Private/PlayAnim/PlayAnimSlotTrait.cpp b/Engine/Plugins/Experimental/UAF/UAFAnimGraph/Source/UAFAnimGraph/Private/PlayAnim/PlayAnimSlotTrait.cpp
--- a/Engine/Plugins/Experimental/UAF/UAFAnimGraph/Source/UAFAnimGraph/Private/PlayAnim/PlayAnimSlotTrait.cpp
+++ b/Engine/Plugins/Experimental/UAF/UAFAnimGraph/Source/UAFAnimGraph/Private/PlayAnim/PlayAnimSlotTrait.cpp
@@ -9,6 +9,7 @@
 #include "TraitInterfaces/ITimeline.h"
 #include "Graph/AnimNextGraphInstance.h"
 #include "Logging/StructuredLog.h"
+#include "Module/AnimNextModuleInstance.h"
 #include "TraitInterfaces/IGraphFactory.h"
 
 #include UE_INLINE_GENERATED_CPP_BY_NAME(PlayAnimSlotTrait)
@@ -383,9 +384,8 @@ namespace UE::UAF
 				FAnimNextGraphInstance& Owner = Binding.GetTraitPtr().GetNodeInstance()->GetOwner();
 				SlotRequest.GraphInstance = SlotRequest.AnimationGraph->AllocateInstance(
 					{
-						.ModuleInstance = Owner.GetModuleInstance(),
-						.ParentContext = &Context,
-						.ParentGraphInstance = &Binding.GetTraitPtr().GetNodeInstance()->GetOwner(),
+						.SystemReference = Owner.GetModuleInstance()->GetReference(),
+						.ParentGraphInstance = Owner.AsShared(),
 						.EntryPoint = EntryPoint
 					});
 
diff --git a/Engine/Plugins/Experimental/UAF/UAFAnimGraph/Source/UAFAnimGraph/Private/TraitCore/ExecutionContext.cpp b/Engine/Plugins/Experimental/UAF/UAFAnimGraph/Source/UAFAnimGraph/Private/TraitCore/ExecutionContext.cpp
--- a/Engine/Plugins/Experimental/UAF/UAFAnimGraph/Source/UAFAnimGraph/Private/TraitCore/ExecutionContext.cpp
+++ b/Engine/Plugins/Experimental/UAF/UAFAnimGraph/Source/UAFAnimGraph/Private/TraitCore/ExecutionContext.cpp
@@ -13,7 +13,7 @@
 #include "Graph/AnimNextAnimationGraph.h"
 #include "Graph/AnimNextGraphInstance.h"
 #include "Components/SceneComponent.h"
-#include "Graph/UAFOutputPoseComponent.h"
+#include "Graph/UAFSystemOutputComponent.h"
 #include "Module/AnimNextModuleInstance.h"
 #include "Logging/StructuredLog.h"
 
@@ -603,6 +603,20 @@ namespace UE::UAF
 		return nullptr;
 	}
 
+	TWeakObjectPtr<const USkeletalMeshComponent> FExecutionContext::GetBindingObject() const
+	{
+		// Pull the component from the system's output component
+		if (IsBound())
+		{
+			if (FAnimNextModuleInstance* SystemInstance = RootGraphInstance->GetModuleInstance())
+			{
+				return SystemInstance->GetComponent<FUAFSystemOutputComponent>().GetSkeletalMeshComponent();
+			}
+		}
+		
+		return TWeakObjectPtr<const USkeletalMeshComponent>();
+	}
+
 #if UE_ENABLE_DEBUG_DRAWING
 	FRigVMDrawInterface* FExecutionContext::GetDebugDrawInterface() const
 	{
diff --git a/Engine/Plugins/Experimental/UAF/UAFAnimGraph/Source/UAFAnimGraph/Private/TraitInterfaces/IEvaluate.cpp b/Engine/Plugins/Experimental/UAF/UAFAnimGraph/Source/UAFAnimGraph/Private/TraitInterfaces/IEvaluate.cpp
--- a/Engine/Plugins/Experimental/UAF/UAFAnimGraph/Source/UAFAnimGraph/Private/TraitInterfaces/IEvaluate.cpp
+++ b/Engine/Plugins/Experimental/UAF/UAFAnimGraph/Source/UAFAnimGraph/Private/TraitInterfaces/IEvaluate.cpp
@@ -11,6 +11,8 @@
 #include "Graph/AnimNextGraphInstance.h"
 #include "Graph/AnimNext_LODPose.h"
 #include "Graph/UAFAnimGraphTrace.h"
+#include "Graph/UAFSystemOutputComponent.h"
+#include "Module/AnimNextModuleInstance.h"
 
 #include "UAF/Attributes/EngineAttributes.h"
 #include "UAF/ValueRuntime/ValueBundle.h"
@@ -61,6 +63,16 @@ namespace UE::UAF
 		}
 	}
 
+	TWeakObjectPtr<const USkeletalMeshComponent> FEvaluateGraphContext::GetBindingObject() const
+	{
+		// Pull the component from the system's output component
+		if (FAnimNextModuleInstance* SystemInstance = GraphInstance.GetModuleInstance())
+		{
+			return SystemInstance->GetComponent<FUAFSystemOutputComponent>().GetSkeletalMeshComponent();
+		}
+		return TWeakObjectPtr<const USkeletalMeshComponent>();
+	}
+
 	enum class EEvaluateStep
 	{
 		PreEvaluate,
@@ -123,7 +135,6 @@ namespace UE::UAF
 		FChildrenArray Children;
 
 		FEvaluateTraversalContext TraversalContext(EvaluationProgram);
-		TraversalContext.SetBindingObject(EvaluateGraphContext.GetBindingObject());
 		TraversalContext.BindTo(GraphRootPtr);
 
 		// Add the graph root to kick start the evaluation process
diff --git a/Engine/Plugins/Experimental/UAF/UAFAnimGraph/Source/UAFAnimGraph/Private/TraitInterfaces/IUpdate.cpp b/Engine/Plugins/Experimental/UAF/UAFAnimGraph/Source/UAFAnimGraph/Private/TraitInterfaces/IUpdate.cpp
--- a/Engine/Plugins/Experimental/UAF/UAFAnimGraph/Source/UAFAnimGraph/Private/TraitInterfaces/IUpdate.cpp
+++ b/Engine/Plugins/Experimental/UAF/UAFAnimGraph/Source/UAFAnimGraph/Private/TraitInterfaces/IUpdate.cpp
@@ -6,6 +6,8 @@
 #include "Graph/UAFGraphInstanceComponent.h"
 #include "AnimNextAnimGraphStats.h"
 #include "Graph/AnimNextGraphInstance.h"
+#include "Graph/UAFSystemOutputComponent.h"
+#include "Module/AnimNextModuleInstance.h"
 #include "TraitCore/TraitEventRaising.h"
 
 namespace UE::UAF
@@ -220,6 +222,16 @@ namespace UE::UAF
 		// No need to forward to our super
 	}
 
+	TWeakObjectPtr<const USkeletalMeshComponent> FUpdateGraphContext::GetBindingObject() const
+	{
+		// Pull the component from the system's output component
+		if (FAnimNextModuleInstance* SystemInstance = GraphInstance.GetModuleInstance())
+		{
+			return SystemInstance->GetComponent<FUAFSystemOutputComponent>().GetSkeletalMeshComponent();
+		}
+		return TWeakObjectPtr<const USkeletalMeshComponent>();
+	}
+
 	//////////////////////////////////////////////////////////////////////////
 	// Traversal implementation
 
@@ -495,7 +507,6 @@ namespace UE::UAF
 		FUpdateTraversalContext TraversalContext;
 		TraversalContext.InputEventList = &InputEventList;
 		TraversalContext.OutputEventList = &OutputEventList;
-		TraversalContext.SetBindingObject(UpdateGraphContext.GetBindingObject());
 
 		FMemStack& MemStack = TraversalContext.GetMemStack();
 		FMemMark Mark(MemStack);
diff --git a/Engine/Plugins/Experimental/UAF/UAFAnimGraph/Source/UAFAnimGraph/Private/Traits/BlendSpacePlayer.cpp b/Engine/Plugins/Experimental/UAF/UAFAnimGraph/Source/UAFAnimGraph/Private/Traits/BlendSpacePlayer.cpp
--- a/Engine/Plugins/Experimental/UAF/UAFAnimGraph/Source/UAFAnimGraph/Private/Traits/BlendSpacePlayer.cpp
+++ b/Engine/Plugins/Experimental/UAF/UAFAnimGraph/Source/UAFAnimGraph/Private/Traits/BlendSpacePlayer.cpp
@@ -143,13 +143,15 @@ namespace UE::UAF
 						FAnimNextGraphInstance& Owner = Binding.GetTraitPtr().GetNodeInstance()->GetOwner();
 						TargetGraphState.Instance = AnimationGraph->AllocateInstance(
 							{
-								.ModuleInstance = Owner.GetModuleInstance(),
-								.ParentContext = &Context,
-								.ParentGraphInstance = &Binding.GetTraitPtr().GetNodeInstance()->GetOwner(),
-								.EntryPoint = EntryPoint
-							});
-						TargetGraphState.ChildPtr = TargetGraphState.Instance->GetGraphRootPtr();
-						FactoryParams.InitializeInstance(*TargetGraphState.Instance.Get());
+								.SystemReference = Owner.GetModuleInstance()->GetReference(),
+								.ParentGraphInstance = Owner.AsShared(),
+								.EntryPoint = EntryPoint,
+							},
+							MoveTemp(FactoryParams));
+						if (TargetGraphState.Instance.IsValid())
+						{
+							TargetGraphState.ChildPtr = TargetGraphState.Instance->GetGraphRootPtr();
+						}
 					}
 				}
 			}
diff --git a/Engine/Plugins/Experimental/UAF/UAFAnimGraph/Source/UAFAnimGraph/Private/Traits/BlendStackTrait.cpp b/Engine/Plugins/Experimental/UAF/UAFAnimGraph/Source/UAFAnimGraph/Private/Traits/BlendStackTrait.cpp
--- a/Engine/Plugins/Experimental/UAF/UAFAnimGraph/Source/UAFAnimGraph/Private/Traits/BlendStackTrait.cpp
+++ b/Engine/Plugins/Experimental/UAF/UAFAnimGraph/Source/UAFAnimGraph/Private/Traits/BlendStackTrait.cpp
@@ -13,6 +13,7 @@
 #include "TraitInterfaces/IInertializerBlend.h"
 #include "HierarchyTableBlendProfile.h"
 #include "Factory/AnimGraphFactory.h"
+#include "Module/AnimNextModuleInstance.h"
 
 #include UE_INLINE_GENERATED_CPP_BY_NAME(BlendStackTrait)
 
@@ -216,16 +217,14 @@ namespace UE::UAF
 					FAnimNextGraphInstance& Owner = Binding.GetTraitPtr().GetNodeInstance()->GetOwner();
 					Graph.Instance = Graph.Request.AnimationGraph->AllocateInstance(
 						{
-							.ModuleInstance = Owner.GetModuleInstance(),
-							.ParentContext = &Context,
-							.ParentGraphInstance = &Owner,
-						});
+							.SystemReference = Owner.GetModuleInstance()->GetReference(),
+							.ParentGraphInstance = Owner.AsShared(),
+						},
+						MoveTemp(Graph.Request.FactoryParams));
 					if (Graph.Request.FactoryObject)
 					{
 						Graph.Instance->SetDebugName(Graph.Request.FactoryObject->GetFName());
 					}
-					Graph.Request.FactoryParams.InitializeInstance(*Graph.Instance.Get());
-
 					break;
 				}
 			case EGraphRequestType::Child:
diff --git a/Engine/Plugins/Experimental/UAF/UAFAnimGraph/Source/UAFAnimGraph/Private/Traits/InputPoseTrait.cpp b/Engine/Plugins/Experimental/UAF/UAFAnimGraph/Source/UAFAnimGraph/Private/Traits/InputPoseTrait.cpp
--- a/Engine/Plugins/Experimental/UAF/UAFAnimGraph/Source/UAFAnimGraph/Private/Traits/InputPoseTrait.cpp
+++ b/Engine/Plugins/Experimental/UAF/UAFAnimGraph/Source/UAFAnimGraph/Private/Traits/InputPoseTrait.cpp
@@ -3,6 +3,7 @@
 #include "Traits/InputPoseTrait.h"
 #include "TraitCore/ExecutionContext.h"
 #include "EvaluationVM/Tasks/PushPose.h"
+#include "Traits/InputValueTrait.h"
 
 #include UE_INLINE_GENERATED_CPP_BY_NAME(InputPoseTrait)
 
@@ -27,4 +28,9 @@ namespace UE::UAF
 			Context.AppendTask(FAnimNextPushPoseTask::Make(&InputPose));
 		}
 	}
+
+	FRigVMStructUpgradeInfo FInputPoseTrait::GetUpgradeInfo() const
+	{
+		return FRigVMStructUpgradeInfo(FAnimNextInputPoseTraitSharedData(), FUAFInputValueTraitSharedData());
+	}
 }
diff --git a/Engine/Plugins/Experimental/UAF/UAFAnimGraph/Source/UAFAnimGraph/Private/Traits/InputPoseTrait.h b/Engine/Plugins/Experimental/UAF/UAFAnimGraph/Source/UAFAnimGraph/Private/Traits/InputPoseTrait.h
--- a/Engine/Plugins/Experimental/UAF/UAFAnimGraph/Source/UAFAnimGraph/Private/Traits/InputPoseTrait.h
+++ b/Engine/Plugins/Experimental/UAF/UAFAnimGraph/Source/UAFAnimGraph/Private/Traits/InputPoseTrait.h
@@ -10,7 +10,7 @@
 
 #include "InputPoseTrait.generated.h"
 
-USTRUCT(meta = (DisplayName = "Input Pose", ShowTooltip=true))
+USTRUCT(meta = (Deprecated = "5.8", DisplayName = "Input Pose (Outdated)", ShowTooltip=true))
 struct FAnimNextInputPoseTraitSharedData : public FAnimNextTraitSharedData
 {
 	GENERATED_BODY()
@@ -36,5 +36,8 @@ namespace UE::UAF
 
 		// IEvaluate impl
 		virtual void PreEvaluate(FEvaluateTraversalContext& Context, const TTraitBinding<IEvaluate>& Binding) const override;
+
+		// FTrait interface
+		virtual FRigVMStructUpgradeInfo GetUpgradeInfo() const override;
 	};
 }
\ No newline at end of file
diff --git a/Engine/Plugins/Experimental/UAF/UAFAnimGraph/Source/UAFAnimGraph/Private/Traits/InputValueTrait.cpp b/Engine/Plugins/Experimental/UAF/UAFAnimGraph/Source/UAFAnimGraph/Private/Traits/InputValueTrait.cpp
new file mode 100644
--- /dev/null
+++ b/Engine/Plugins/Experimental/UAF/UAFAnimGraph/Source/UAFAnimGraph/Private/Traits/InputValueTrait.cpp
@@ -0,0 +1,148 @@
+// Copyright Epic Games, Inc. All Rights Reserved.
+
+#include "Traits/InputValueTrait.h"
+#include "TraitCore/ExecutionContext.h"
+#include "EvaluationVM/Tasks/PushPose.h"
+#include "EvaluationVM/Tasks/PushReferenceKeyframe.h"
+#include "Graph/UAFSystemOutputComponent.h"
+#include "Graph/WeakAnimGraphReference.h"
+#include "TraitCore/NodeInstance.h"
+#include "UAF/ValueRuntime/IVirtualValueBundle.h"
+#include "UAF/ValueRuntime/ValueBundle.h"
+
+#include UE_INLINE_GENERATED_CPP_BY_NAME(InputValueTrait)
+
+namespace UE::UAF
+{
+	AUTO_REGISTER_ANIM_TRAIT(FInputValueTrait)
+
+	// Trait implementation boilerplate
+#define TRAIT_INTERFACE_ENUMERATOR(GeneratorMacro) \
+	GeneratorMacro(IUpdate) \
+	GeneratorMacro(IEvaluate) \
+	GeneratorMacro(IUpdateTraversal) \
+	GeneratorMacro(IHierarchy) \
+
+	GENERATE_ANIM_TRAIT_IMPLEMENTATION(FInputValueTrait, TRAIT_INTERFACE_ENUMERATOR, NULL_ANIM_TRAIT_INTERFACE_ENUMERATOR, NULL_ANIM_TRAIT_EVENT_ENUMERATOR)
+#undef TRAIT_INTERFACE_ENUMERATOR
+
+	const FInputValueTrait::FInstanceData* FInputValueTrait::UpdateInstanceData(const FTraitBinding& InBinding) const
+	{
+		const FSharedData* SharedData = InBinding.GetSharedData<FSharedData>();
+		FInstanceData* InstanceData = InBinding.GetInstanceData<FInstanceData>();
+		InstanceData->VariableReference = SharedData->GetInput(InBinding);
+		InstanceData->GraphInstance = &InBinding.GetTraitPtr().GetNodeInstance()->GetOwner();
+		check(InstanceData->GraphInstance);
+		return InstanceData;
+	}
+
+	void FInputValueTrait::OnBecomeRelevant(FUpdateTraversalContext& Context, const TTraitBinding<IUpdate>& Binding, const FTraitUpdateState& TraitState) const
+	{
+		(void)UpdateInstanceData(Binding);
+	}
+
+	void FInputValueTrait::PreUpdate(FUpdateTraversalContext& Context, const TTraitBinding<IUpdate>& Binding, const FTraitUpdateState& TraitState) const
+	{
+		const FInstanceData* InstanceData = UpdateInstanceData(Binding);
+		InstanceData->GraphInstance->AccessVariable<FUAFValueBundle>(InstanceData->VariableReference, [&Context, &Binding, &TraitState](FUAFValueBundle& InValueBundle)
+		{
+			if (const IVirtualValueBundle* Impl = InValueBundle.GetImpl())
+			{
+				if (const FWeakAnimGraphReference* GraphReference = Impl->GetAnimGraphReference())
+				{
+					(void)GraphReference->GetOrAllocate(Context, Binding);
+				}
+			}
+		});
+	}
+
+	void FInputValueTrait::PreEvaluate(FEvaluateTraversalContext& Context, const TTraitBinding<IEvaluate>& Binding) const
+	{
+		bool bEvaluated = false;
+		const FInstanceData* InstanceData = UpdateInstanceData(Binding);
+		InstanceData->GraphInstance->AccessVariable<FUAFValueBundle>(InstanceData->VariableReference, [&bEvaluated, &Context, &Binding](FUAFValueBundle& InValueBundle)
+		{
+			if (const IVirtualValueBundle* Impl = InValueBundle.GetImpl())
+			{
+				// If this value bundle represents a concrete value, then we need to use it, otherwise assume it is virtual and any task generation is
+				// handled by the PreEvaluate call above
+				if (const FAnimNextGraphLODPose* LODPose = Impl->GetLODPose())
+				{
+					// TODO: Implement with new attribute runtime
+					if (LODPose->LODPose.IsValid())
+					{
+						Context.AppendTask(FAnimNextPushPoseTask::Make(LODPose));
+					}
+				}
+
+				bEvaluated = true;
+			}
+
+		});
+
+		if (!bEvaluated)
+		{
+			Context.AppendTask(FAnimNextPushReferenceKeyframeTask::MakeFromSkeleton());
+		}
+	}
+
+	void FInputValueTrait::QueueChildrenForTraversal(FUpdateTraversalContext& Context, const TTraitBinding<IUpdateTraversal>& Binding, const FTraitUpdateState& TraitState, FUpdateTraversalQueue& TraversalQueue) const
+	{
+		const FInstanceData* InstanceData = UpdateInstanceData(Binding);
+		InstanceData->GraphInstance->AccessVariable<FUAFValueBundle>(InstanceData->VariableReference, [&Context, &Binding, &TraitState, &TraversalQueue](FUAFValueBundle& InValueBundle)
+		{
+			if (const IVirtualValueBundle* Impl = InValueBundle.GetImpl())
+			{
+				if (const FWeakAnimGraphReference* GraphReference = Impl->GetAnimGraphReference())
+				{
+					TSharedPtr<FAnimNextGraphInstance> GraphInstance = GraphReference->GetOrAllocate(Context, Binding);
+					if (GraphInstance.IsValid() && GraphInstance->GetGraphRootPtr().IsValid())
+					{
+						TraversalQueue.Push(GraphInstance->GetGraphRootPtr(), TraitState);
+					}
+				}
+			}
+		});
+	}
+	
+	uint32 FInputValueTrait::GetNumChildren(const FExecutionContext& Context, const TTraitBinding<IHierarchy>& Binding) const
+	{
+		uint32 NumChildren = 0;
+		const FInstanceData* InstanceData = UpdateInstanceData(Binding);
+		InstanceData->GraphInstance->AccessVariable<FUAFValueBundle>(InstanceData->VariableReference, [&Context, &Binding, &NumChildren](FUAFValueBundle& InValueBundle)
+		{
+			if (const IVirtualValueBundle* Impl = InValueBundle.GetImpl())
+			{
+				if (const FWeakAnimGraphReference* GraphReference = Impl->GetAnimGraphReference())
+				{
+					TSharedPtr<FAnimNextGraphInstance> GraphInstance = GraphReference->GetOrAllocate(Context, Binding);
+					if (GraphInstance.IsValid() && GraphInstance->GetGraphRootPtr().IsValid())
+					{
+						NumChildren = 1;
+					}
+				}
+			}
+		});
+
+		return NumChildren;
+	}
+
+	void FInputValueTrait::GetChildren(const FExecutionContext& Context, const TTraitBinding<IHierarchy>& Binding, FChildrenArray& Children) const
+	{
+		const FInstanceData* InstanceData = UpdateInstanceData(Binding);
+		InstanceData->GraphInstance->AccessVariable<FUAFValueBundle>(InstanceData->VariableReference, [&Context, &Binding, &Children](FUAFValueBundle& InValueBundle)
+		{
+			if (const IVirtualValueBundle* Impl = InValueBundle.GetImpl())
+			{
+				if (const FWeakAnimGraphReference* GraphReference = Impl->GetAnimGraphReference())
+				{
+					TSharedPtr<FAnimNextGraphInstance> GraphInstance = GraphReference->GetOrAllocate(Context, Binding);
+					if (GraphInstance.IsValid() && GraphInstance->GetGraphRootPtr().IsValid())
+					{
+						Children.Add(GraphInstance->GetGraphRootPtr());
+					}
+				}
+			}
+		});
+	}
+}
diff --git a/Engine/Plugins/Experimental/UAF/UAFAnimGraph/Source/UAFAnimGraph/Private/Traits/SubGraphHost.cpp b/Engine/Plugins/Experimental/UAF/UAFAnimGraph/Source/UAFAnimGraph/Private/Traits/SubGraphHost.cpp
--- a/Engine/Plugins/Experimental/UAF/UAFAnimGraph/Source/UAFAnimGraph/Private/Traits/SubGraphHost.cpp
+++ b/Engine/Plugins/Experimental/UAF/UAFAnimGraph/Source/UAFAnimGraph/Private/Traits/SubGraphHost.cpp
@@ -5,6 +5,7 @@
 #include "TraitCore/ExecutionContext.h"
 #include "TraitCore/NodeInstance.h"
 #include "Graph/AnimNextGraphInstance.h"
+#include "Module/AnimNextModuleInstance.h"
 
 namespace UE::UAF
 {
@@ -215,9 +216,8 @@ namespace UE::UAF
 				FAnimNextGraphInstance& Owner = Binding.GetTraitPtr().GetNodeInstance()->GetOwner();
 				SubGraphEntry.GraphInstance = SubGraphEntry.AnimationGraph->AllocateInstance(
 					{
-						.ModuleInstance = Owner.GetModuleInstance(),
-						.ParentContext = &Context,
-						.ParentGraphInstance = &Binding.GetTraitPtr().GetNodeInstance()->GetOwner(),
+						.SystemReference = Owner.GetModuleInstance()->GetReference(),
+						.ParentGraphInstance = Owner.AsShared(),
 						.EntryPoint = SubGraphEntry.EntryPoint
 					});
 			}
diff --git a/Engine/Plugins/Experimental/UAF/UAFAnimGraph/Source/UAFAnimGraph/Public/Factory/AnimNextFactoryParams.h b/Engine/Plugins/Experimental/UAF/UAFAnimGraph/Source/UAFAnimGraph/Public/Factory/AnimNextFactoryParams.h
--- a/Engine/Plugins/Experimental/UAF/UAFAnimGraph/Source/UAFAnimGraph/Public/Factory/AnimNextFactoryParams.h
+++ b/Engine/Plugins/Experimental/UAF/UAFAnimGraph/Source/UAFAnimGraph/Public/Factory/AnimNextFactoryParams.h
@@ -150,9 +150,6 @@ struct FAnimNextFactoryParams
 		return *this;
 	}
 
-	// Apply these params to the supplied instance.
-	UAFANIMGRAPH_API void InitializeInstance(FUAFAssetInstance& InInstance) const;
-
 	// Clear these parameters back to empty
 	UAFANIMGRAPH_API void Reset();
 
@@ -160,6 +157,11 @@ private:
 	friend UE::UAF::FAnimGraphFactory;
 	friend UE::UAF::Editor::FAnimNextFactoryParamsDetails;
 	friend FAnimNextGraphInstanceTaskInstanceData;	// For old data upgrade
+	friend UUAFAnimGraph;
+	friend FAnimNextGraphInstance;
+
+	// Apply these params to the supplied instance.
+	UAFANIMGRAPH_API void InitializeInstance(FUAFAssetInstance& InInstance) const;
 
 	// Method used to create a graph
 	UPROPERTY(EditAnywhere, Category = Parameters)
diff --git a/Engine/Plugins/Experimental/UAF/UAFAnimGraph/Source/UAFAnimGraph/Public/Graph/AnimGraphReference.h b/Engine/Plugins/Experimental/UAF/UAFAnimGraph/Source/UAFAnimGraph/Public/Graph/AnimGraphReference.h
new file mode 100644
--- /dev/null
+++ b/Engine/Plugins/Experimental/UAF/UAFAnimGraph/Source/UAFAnimGraph/Public/Graph/AnimGraphReference.h
@@ -0,0 +1,86 @@
+﻿// Copyright Epic Games, Inc. All Rights Reserved.
+
+#pragma once
+
+#include "AnimGraphTaskContext.h"
+#include "EngineDefines.h"
+#include "Factory/AnimNextFactoryParams.h"
+#include "Templates/NonNullPointer.h"
+#include "Templates/SharedPointer.h"
+
+#define UE_API UAFANIMGRAPH_API
+
+struct FUAFWeakSystemReference;
+class UUAFAnimGraph;
+struct FAnimNextGraphInstance;
+class FReferenceCollector;
+
+namespace UE::UAF
+{
+struct FTraitBinding;
+struct FExecutionContext;
+struct FWeakTraitPtr;
+struct FWeakAnimGraphReference;
+}
+
+namespace UE::UAF
+{
+
+// An owning reference to a UAF graph.
+// Provides a way for external systems to own an anim graph, but not have to fully allocate the graph before use
+struct FAnimGraphReference
+{
+	// Default-constructable only because it needs to be embedded in UObjects/UStructs
+	FAnimGraphReference() = default;
+
+	// Construct an owning handle to a graph instance.
+	// Constructs a FAnimNextGraphInstance for use by worker threads and minimally pre-allocates it
+	// Can be called on any thread
+	// @param    InAnimGraph        The graph asset to create this instance from
+	// @param    InSystemReference  The system that the graph will be used with. Graphs can only be used within a single system
+	// @param    InFactoryParams    The factory params to construct the instance with
+	UE_API explicit FAnimGraphReference(TNonNullPtr<const UUAFAnimGraph> InAnimGraph, const FUAFWeakSystemReference& InSystemReference, FAnimNextFactoryParams&& InFactoryParams = FAnimNextFactoryParams());
+
+	// Make a weak reference from this owning reference
+	UE_API FWeakAnimGraphReference ToWeakRef() const;
+
+	// Reset this handle to invalid - gives up ownership
+	// Can be called on any thread
+	UE_API void Reset();
+
+	// Check if this handle is valid
+	// Can be called on any thread.
+	UE_API bool IsValid() const;
+
+	// Get the anim graph that this handle references
+	// Can be called on any thread.
+	[[nodiscard]] UE_API const UUAFAnimGraph* GetAnimGraph() const;
+
+	// Queue task to execute on the instance
+	// Can be called on any thread.
+	// @param	InTaskFunction		The function to run
+	UE_API void QueueTask(FUniqueAnimGraphTask&& InTaskFunction) const;
+
+	// Sets the debug name for the instance, used for tracing etc.
+	UE_API void SetDebugName(FName InName) const;
+
+	// GC support
+	UE_API void AddReferencedObjects(FReferenceCollector& InCollector) const;
+
+private:
+	// Gets or allocate & binds the graph instance
+	// Fully allocates the graph if it is only preallocated
+	// Can be called on any thread.
+	// @param    InBinding    The current binding (this must be called during a graph execution)
+	[[nodiscard]] UE_API TSharedPtr<FAnimNextGraphInstance> GetOrAllocate(FExecutionContext& InContext, const FTraitBinding& InBinding) const;
+
+private:
+	friend FWeakAnimGraphReference;
+	
+	// Pointer to the instance
+	TSharedPtr<FAnimNextGraphInstance> Ptr;
+};
+
+}
+
+#undef UE_API
\ No newline at end of file
diff --git a/Engine/Plugins/Experimental/UAF/UAFAnimGraph/Source/UAFAnimGraph/Public/Graph/AnimGraphTaskContext.h b/Engine/Plugins/Experimental/UAF/UAFAnimGraph/Source/UAFAnimGraph/Public/Graph/AnimGraphTaskContext.h
new file mode 100644
--- /dev/null
+++ b/Engine/Plugins/Experimental/UAF/UAFAnimGraph/Source/UAFAnimGraph/Public/Graph/AnimGraphTaskContext.h
@@ -0,0 +1,40 @@
+﻿// Copyright Epic Games, Inc. All Rights Reserved.
+
+#pragma once
+
+#include "InstanceTaskContext.h"
+#include "Templates/Function.h"
+
+#define UE_API UAFANIMGRAPH_API 
+
+struct FAnimNextGraphInstance;
+
+namespace UE::UAF
+{
+struct FAnimGraphReference;
+struct FWeakAnimGraphReference;
+}
+
+namespace UE::UAF
+{
+
+// Context struct passed to graph instance tasks, allowing setting of variables etc.
+struct FAnimGraphTaskContext : public FInstanceTaskContext
+{
+public:
+	// Get the graph instance
+	UE_API FAnimNextGraphInstance& GetAnimGraphInstance() const;
+
+private:
+	friend FAnimNextGraphInstance;
+	friend UE::UAF::FAnimGraphReference;
+	friend UE::UAF::FWeakAnimGraphReference;
+
+	UE_API FAnimGraphTaskContext(FAnimNextGraphInstance& InGraphInstance);
+};
+
+using FUniqueAnimGraphTask = ::TUniqueFunction<void(const FAnimGraphTaskContext& InContext)>;
+
+}
+
+#undef UE_API
\ No newline at end of file
diff --git a/Engine/Plugins/Experimental/UAF/UAFAnimGraph/Source/UAFAnimGraph/Public/Graph/AnimNextAnimGraph.h b/Engine/Plugins/Experimental/UAF/UAFAnimGraph/Source/UAFAnimGraph/Public/Graph/AnimNextAnimGraph.h
--- a/Engine/Plugins/Experimental/UAF/UAFAnimGraph/Source/UAFAnimGraph/Public/Graph/AnimNextAnimGraph.h
+++ b/Engine/Plugins/Experimental/UAF/UAFAnimGraph/Source/UAFAnimGraph/Public/Graph/AnimNextAnimGraph.h
@@ -16,6 +16,7 @@ struct FAnimNextGraphInstance;
 struct FRigUnit_AnimNextGraphEvaluator;
 struct FRigUnit_AnimNextRunAnimationGraph_v1;
 struct FRigUnit_AnimNextRunAnimationGraph_v2;
+struct FRigUnit_UAFRunAsset;
 class UUAFAnimGraph;
 class FRigVMTraitScope;
 struct FRigVMExtendedExecuteContext;
@@ -24,14 +25,12 @@ struct FAnimNextModuleAnimGraphComponent;
 struct FStructView;
 struct FAnimNextModuleInjectionComponent;
 struct FRigUnit_GetPostProcessAnimation;
-struct FEntitySkeletalAnimationStackComponent;
 
 namespace UE::UAF
 {
 	struct IEvaluationModifier;
 	struct FExecutionContext;
 	struct FGraphInstanceComponent;
-	class IDataInterfaceHost;
 	struct FBlendStackCoreTrait;
 	struct FInjectionSiteTrait;
 	struct FInjectionInfo;
@@ -99,6 +98,7 @@ private:
 
 	friend FRigUnit_AnimNextRunAnimationGraph_v1;	// We evaluate the instance
 	friend FRigUnit_AnimNextRunAnimationGraph_v2;	// We evaluate the instance
+	friend FRigUnit_UAFRunAsset;					// We evaluate the instance
 	friend UE::UAF::FBlendStackCoreTrait;		// Requires access to take a reference to external graphs
 	friend UE::UAF::FInjectionSiteTrait;		// Requires access to take a reference to external graphs
 	friend UE::UAF::FInjectionInfo;
@@ -106,7 +106,6 @@ private:
 	friend UE::UAF::Editor::FAnimNextGraphDetails;
 	friend FRigUnit_GetPostProcessAnimation;
 	friend UE::UAF::AnimGraph::FAnimNextAnimGraphModule;
-	friend FEntitySkeletalAnimationStackComponent;
 };
 
 #undef UE_API
diff --git a/Engine/Plugins/Experimental/UAF/UAFAnimGraph/Source/UAFAnimGraph/Public/Graph/AnimNextAnimationGraph.h b/Engine/Plugins/Experimental/UAF/UAFAnimGraph/Source/UAFAnimGraph/Public/Graph/AnimNextAnimationGraph.h
--- a/Engine/Plugins/Experimental/UAF/UAFAnimGraph/Source/UAFAnimGraph/Public/Graph/AnimNextAnimationGraph.h
+++ b/Engine/Plugins/Experimental/UAF/UAFAnimGraph/Source/UAFAnimGraph/Public/Graph/AnimNextAnimationGraph.h
@@ -11,6 +11,8 @@
 #include "TraitCore/EntryPointHandle.h"
 #include "RigVMHost.h"
 #include "AnimNextGraphEvaluatorExecuteDefinition.h"
+#include "Factory/AnimNextFactoryParams.h"
+#include "Module/UAFWeakSystemReference.h"
 
 #include "AnimNextAnimationGraph.generated.h"
 
@@ -70,12 +72,10 @@ namespace UE::UAF
 	// Parameters to AllocateInstance
 	struct FGraphAllocationParams
 	{
-		// The module instance to use as our outermost host. Usually non-null except in special cases (e.g. unit tests)
-		FAnimNextModuleInstance* ModuleInstance = nullptr;
-		// The current execution context we are running in, if any
-		FExecutionContext* ParentContext = nullptr;
+		// The system to use as our outermost host
+		FUAFWeakSystemReference SystemReference;
 		// The hosting/parent graph instance
-		FAnimNextGraphInstance* ParentGraphInstance = nullptr;
+		TWeakPtr<FAnimNextGraphInstance> ParentGraphInstance;
 		// The entry point to use (deprecated)
 		FName EntryPoint = NAME_None;
 		// Any variable overrides to apply to this graph
@@ -101,10 +101,17 @@ public:
 	UE_API virtual void GetPreloadDependencies(TArray<UObject*>& OutDeps) override;
 #endif
 
+	// Pre-allocates an instance of the graph
+	// @param	InParams		Parameters to pre-allocate the graph with
+	// @param	InFactoryParams	Factory parameters to apply to the graph when allocated
+	// @return the pre-allocated instance.
+	UE_API TSharedPtr<FAnimNextGraphInstance> PreAllocateInstance(const UE::UAF::FGraphAllocationParams& InParams = UE::UAF::FGraphAllocationParams(), FAnimNextFactoryParams&& InFactoryParams = FAnimNextFactoryParams()) const;
+
 	// Allocates an instance of the graph according to the supplied parameters
 	// @param	InParams		Parameters to allocate the graph with
+	// @param	InFactoryParams	Factory parameters to apply to the graph when allocated
 	// @return the allocated instance. This can be invalid if the entry point was not found or the graph's root trait is invalid.
-	UE_API TSharedPtr<FAnimNextGraphInstance> AllocateInstance(const UE::UAF::FGraphAllocationParams& InParams = UE::UAF::FGraphAllocationParams()) const;
+	UE_API TSharedPtr<FAnimNextGraphInstance> AllocateInstance(const UE::UAF::FGraphAllocationParams& InParams = UE::UAF::FGraphAllocationParams(), FAnimNextFactoryParams&& InFactoryParams = FAnimNextFactoryParams()) const;
 
 private:
 	// Loads the graph data from the provided archive buffer and returns true on success, false otherwise
diff --git a/Engine/Plugins/Experimental/UAF/UAFAnimGraph/Source/UAFAnimGraph/Public/Graph/AnimNextGraphInstance.h b/Engine/Plugins/Experimental/UAF/UAFAnimGraph/Source/UAFAnimGraph/Public/Graph/AnimNextGraphInstance.h
--- a/Engine/Plugins/Experimental/UAF/UAFAnimGraph/Source/UAFAnimGraph/Public/Graph/AnimNextGraphInstance.h
+++ b/Engine/Plugins/Experimental/UAF/UAFAnimGraph/Source/UAFAnimGraph/Public/Graph/AnimNextGraphInstance.h
@@ -3,7 +3,10 @@
 #pragma once
 
 #include "CoreMinimal.h"
+#include "AnimGraphTaskContext.h"
+#include "InstanceTask.h"
 #include "UAFAssetInstance.h"
+#include "Containers/TransactionallySafeMpscQueue.h"
 #include "HAL/CriticalSection.h"
 #include "TraitCore/TraitEvent.h"
 #include "TraitCore/TraitEventList.h"
@@ -28,10 +31,11 @@ struct FRigUnit_AnimNextRunAnimationGraph_v1;
 struct FRigUnit_AnimNextRunAnimationGraph_v2;
 struct FAnimNextStateTreeRigVMConditionBase;
 struct FAnimNextStateTreeRigVMTaskBase;
+struct FAnimNextFactoryParams;
 
 namespace UE::UAF
 {
-	class IDataInterfaceHost;
+	struct FTraitBinding;
 	struct FExecutionContext;
 	struct FGraphInstanceComponent;
 	struct FLatentPropertyHandle;
@@ -44,6 +48,8 @@ namespace UE::UAF
 	struct FCachedBindingInfo;
 	struct FInjectionUtils;
 	struct FBlendSpacePlayerTrait;
+	struct FAnimGraphReference;
+	struct FWeakAnimGraphReference;
 }
 
 // Represents an instance of an AnimNext graph
@@ -51,6 +57,9 @@ namespace UE::UAF
 // alive while we own a reference to it. It is not intended to be serialized on disk with a live instance.
 USTRUCT()
 struct FAnimNextGraphInstance : public FUAFAssetInstance
+#if CPP
+	, public TSharedFromThis<FAnimNextGraphInstance>
+#endif
 {
 	GENERATED_BODY()
 
@@ -114,6 +123,12 @@ private:
 	// When frozen, latent handles that can freeze are skipped, all others will execute
 	void CopyVariablesToLatentPins(const TConstArrayView<UE::UAF::FLatentPropertyHandle>& LatentHandles, void* DestinationBasePtr, bool bIsFrozen, bool bJustBecameRelevant);
 
+	// Performs allocation of internal variable data structures
+	void AllocateVariablesInternal(const TSharedPtr<const UE::UAF::FVariableOverridesCollection>& InOverrides, TWeakPtr<FAnimNextGraphInstance> InParentGraphInstance);
+
+	// Performs allocation of graph root ptr
+	void AllocateGraphInternal();
+
 #if WITH_EDITORONLY_DATA
 	// During graph compilation, if we have existing graph instances, we freeze them by releasing their memory before thawing them
 	// Freezing is a partial release of resources that retains the necessary information to re-create things safely
@@ -130,6 +145,9 @@ private:
 	// Unregisters the provided trait from the GC system
 	void UnregisterWithGC(const UE::UAF::FWeakTraitPtr& InTraitPtr);
 
+	// Runs a task or queues it until AllocateVariablesInternal
+	void RunOrQueueTask(UE::UAF::FUniqueAnimGraphTask&& InTaskFunction);
+
 private:
 	// GC handler. This is nested to avoid having to bubble the polymorphism it introduces up to super-structs
 	struct FGCHandler : public FGCObject
@@ -149,11 +167,26 @@ private:
 	UE::UAF::FTraitPtr GraphInstancePtr;
 
 	// The module instance that owns the root, us and the components
-	FAnimNextModuleInstance* ModuleInstance = nullptr;
+	TWeakPtr<FAnimNextModuleInstance> ModuleInstance;
 
 	// The root graph instance that owns us and the components
 	FAnimNextGraphInstance* RootGraphInstance = nullptr;
-	
+
+	enum EAllocationState
+	{
+		// Default-initialized graph instance
+		Empty,
+
+		// Has a reference to its graph, but has no internal other structures allocated
+		PreAllocated,
+
+		// Has internal variables & components allocated
+		VariablesAllocated,
+
+		// Has all internal structures allocated
+		Allocated,
+	} AllocationState = EAllocationState::Empty;
+
 	// Whether or not this graph has updated once
 	bool bHasUpdatedOnce : 1 = false;
 
@@ -173,6 +206,8 @@ private:
 	friend FAnimNextStateTreeRigVMConditionBase;	// Temp - remove: Needs to access mutable variables to copy-in data for function shim
 	friend FAnimNextStateTreeRigVMTaskBase;	// Temp - remove: Needs to access mutable variables to copy-in data for function shim
 	friend UE::UAF::FBlendSpacePlayerTrait;
+	friend UE::UAF::FAnimGraphReference;
+	friend UE::UAF::FWeakAnimGraphReference;
 };
 
 template<>
diff --git a/Engine/Plugins/Experimental/UAF/UAFAnimGraph/Source/UAFAnimGraph/Public/Graph/WeakAnimGraphReference.h b/Engine/Plugins/Experimental/UAF/UAFAnimGraph/Source/UAFAnimGraph/Public/Graph/WeakAnimGraphReference.h
new file mode 100644
--- /dev/null
+++ b/Engine/Plugins/Experimental/UAF/UAFAnimGraph/Source/UAFAnimGraph/Public/Graph/WeakAnimGraphReference.h
@@ -0,0 +1,83 @@
+﻿// Copyright Epic Games, Inc. All Rights Reserved.
+
+#pragma once
+
+#include "AnimGraphTaskContext.h"
+#include "EngineDefines.h"
+#include "InstanceTask.h"
+#include "Templates/SharedPointer.h"
+
+#define UE_API UAFANIMGRAPH_API
+
+class UUAFAnimGraph;
+struct FAnimNextGraphInstance;
+class FReferenceCollector;
+
+namespace UE::UAF
+{
+struct FTraitBinding;
+struct FExecutionContext;
+struct FWeakTraitPtr;
+struct FAnimGraphReference;
+struct FInputValueTrait;
+struct FLayerStackTrait;
+}
+
+namespace UE::UAF
+{
+
+// A non-owning reference to a UAF graph.
+struct FWeakAnimGraphReference
+{
+	// Default-constructable only because it needs to be embedded in UObjects/UStructs
+	FWeakAnimGraphReference() = default;
+
+	// Construct a non-owning reference to a graph instance.
+	// Can be called on any thread
+	UE_API explicit FWeakAnimGraphReference(const FAnimGraphReference& InAnimGraphReference);
+
+	// Construct a non-owning reference to a graph instance.
+	// Can be called on any thread
+	UE_API explicit FWeakAnimGraphReference(const TSharedPtr<FAnimNextGraphInstance>& InAnimGraphInstance);
+
+	// Reset this handle to invalid - gives up ownership
+	// Can be called on any thread
+	UE_API void Reset();
+
+	// Check if this handle is valid
+	// Can be called on any thread.
+	UE_API bool IsValid() const;
+
+	// Get the anim graph that this handle references
+	// Can be called on any thread.
+	[[nodiscard]] UE_API const UUAFAnimGraph* GetAnimGraph() const;
+
+	// Queue task to execute on the instance
+	// Can be called on any thread.
+	// @param	InTaskFunction		The function to run
+	UE_API void QueueTask(FUniqueAnimGraphTask&& InTaskFunction) const;
+
+	// Sets the debug name for the instance, used for tracing etc.
+	UE_API void SetDebugName(FName InName) const;
+
+	// GC support
+	UE_API void AddReferencedObjects(FReferenceCollector& InCollector) const;
+
+private:
+	friend FLayerStackTrait;
+	friend FInputValueTrait;
+
+	// Gets or allocate & binds the graph instance
+	// Fully allocates the graph if it is only preallocated
+	// Can be called on any thread.
+	// @param    InBinding    The current binding (this must be called during a graph execution)
+	[[nodiscard]] UE_API TSharedPtr<FAnimNextGraphInstance> GetOrAllocate(const FExecutionContext& InContext, const FTraitBinding& InBinding) const;
+
+private:
+	// Pointer to the instance
+	TWeakPtr<FAnimNextGraphInstance> Ptr;
+};
+
+}
+
+#undef UE_API
\ No newline at end of file
diff --git a/Engine/Plugins/Experimental/UAF/UAFAnimGraph/Source/UAFAnimGraph/Public/Native/UAFSimplePrePhysicsGraphComponent.h b/Engine/Plugins/Experimental/UAF/UAFAnimGraph/Source/UAFAnimGraph/Public/Native/UAFSimplePrePhysicsGraphComponent.h
--- a/Engine/Plugins/Experimental/UAF/UAFAnimGraph/Source/UAFAnimGraph/Public/Native/UAFSimplePrePhysicsGraphComponent.h
+++ b/Engine/Plugins/Experimental/UAF/UAFAnimGraph/Source/UAFAnimGraph/Public/Native/UAFSimplePrePhysicsGraphComponent.h
@@ -4,7 +4,9 @@
 
 #include "Graph/AnimNext_LODPose.h"
 #include "Graph/RigUnit_AnimNextRunAnimationGraph_v2.h"
+#include "Graph/RigUnit_UAFRunAsset.h"
 #include "Script/UAFScriptComponent.h"
+#include "UAF/ValueRuntime/ValueBundle.h"
 #include "UAFSimplePrePhysicsGraphComponent.generated.h"
 
 #define UE_API UAFANIMGRAPH_API
@@ -29,21 +31,13 @@ private:
 	// FUAFAssetInstanceComponent interface
 	UE_API virtual void OnBindToInstance() override;
 
-	// The mesh we output to
+	// Output value
 	UPROPERTY(Transient)
-	TWeakObjectPtr<USkeletalMeshComponent> WeakMesh;
+	FUAFValueBundle Output;
 
-	// Refpose generated from the mesh
+	// Work data for the 'run asset node'
 	UPROPERTY(Transient)
-	FAnimNextGraphReferencePose RefPose;
-
-	// Output pose
-	UPROPERTY(Transient)
-	FAnimNextGraphLODPose Pose;
-
-	// Work data for the 'rungraph node'
-	UPROPERTY(Transient)
-	FAnimNextRunGraphWorkData_v2 WorkData;
+	FUAFRunAssetWorkData WorkData;
 
 	// Reference to injected graph data in variables
 	FAnimNextAnimGraph* Graph = nullptr;
diff --git a/Engine/Plugins/Experimental/UAF/UAFAnimGraph/Source/UAFAnimGraph/Public/TraitCore/ExecutionContext.h b/Engine/Plugins/Experimental/UAF/UAFAnimGraph/Source/UAFAnimGraph/Public/TraitCore/ExecutionContext.h
--- a/Engine/Plugins/Experimental/UAF/UAFAnimGraph/Source/UAFAnimGraph/Public/TraitCore/ExecutionContext.h
+++ b/Engine/Plugins/Experimental/UAF/UAFAnimGraph/Source/UAFAnimGraph/Public/TraitCore/ExecutionContext.h
@@ -202,9 +202,8 @@ namespace UE::UAF
 		// Returns the bound root graph instance
 		FAnimNextGraphInstance& GetRootGraphInstance() const { return *RootGraphInstance; }
 
-		// Returns the skeletal mesh component associated to the graph execution
-		void SetBindingObject(const TWeakObjectPtr<const USkeletalMeshComponent>& InBindingObject) { BindingObject = InBindingObject; }
-		const TWeakObjectPtr<const USkeletalMeshComponent> GetBindingObject() const { return BindingObject; }
+		// Returns the skeletal mesh component associated to the graph execution, if any
+		UE_API TWeakObjectPtr<const USkeletalMeshComponent> GetBindingObject() const;
 
 		// Returns the local thread memstack
 		FMemStack& GetMemStack() const { return MemStack; }
@@ -232,9 +231,6 @@ namespace UE::UAF
 		UE_API const FTrait* GetTrait(const FTraitTemplate& TraitDesc) const;
 
 	protected:
-		// The skeletal mesh component associated to the graph execution
-		TWeakObjectPtr<const USkeletalMeshComponent> BindingObject = nullptr;
-
 		// The memstack of the local thread
 		FMemStack& MemStack;
 
diff --git a/Engine/Plugins/Experimental/UAF/UAFAnimGraph/Source/UAFAnimGraph/Public/TraitInterfaces/IEvaluate.h b/Engine/Plugins/Experimental/UAF/UAFAnimGraph/Source/UAFAnimGraph/Public/TraitInterfaces/IEvaluate.h
--- a/Engine/Plugins/Experimental/UAF/UAFAnimGraph/Source/UAFAnimGraph/Public/TraitInterfaces/IEvaluate.h
+++ b/Engine/Plugins/Experimental/UAF/UAFAnimGraph/Source/UAFAnimGraph/Public/TraitInterfaces/IEvaluate.h
@@ -14,6 +14,7 @@
 struct FAnimNextGraphLODPose;
 struct FRigUnit_AnimNextRunAnimationGraph_v1;
 struct FRigUnit_AnimNextRunAnimationGraph_v2;
+struct FRigUnit_UAFRunAsset;
 struct FAnimNode_AnimNextGraph;
 class UAbstractSkeletonSetBinding;
 
@@ -44,8 +45,7 @@ namespace UE::UAF
 		FName GetSkeletonSetName() const { return SkeletonSetName; }
 
 	protected:
-		const TWeakObjectPtr<const USkeletalMeshComponent>& GetBindingObject() const { return BindingObject; }
-		void SetBindingObject(const TWeakObjectPtr<const USkeletalMeshComponent>& InBindingObject) { BindingObject = InBindingObject; }
+		UE_API TWeakObjectPtr<const USkeletalMeshComponent> GetBindingObject() const;
 
 		void SetSkeletonSet(UAbstractSkeletonSetBinding* InSkeletonSetBinding, FName InSkeletonSetName)
 		{
@@ -62,11 +62,10 @@ namespace UE::UAF
 
 		int32 GraphLODLevel = 0;
 
-		TWeakObjectPtr<const USkeletalMeshComponent> BindingObject = nullptr;
-
 		friend UE_API FEvaluationProgram EvaluateGraph(const FEvaluateGraphContext& EvaluateGraphContext, const FWeakTraitPtr& GraphRootPtr);
 		friend struct ::FRigUnit_AnimNextRunAnimationGraph_v1;
 		friend struct ::FRigUnit_AnimNextRunAnimationGraph_v2;
+		friend struct ::FRigUnit_UAFRunAsset;
 		friend struct ::FAnimNode_AnimNextGraph;
 	};
 
diff --git a/Engine/Plugins/Experimental/UAF/UAFAnimGraph/Source/UAFAnimGraph/Public/TraitInterfaces/IUpdate.h b/Engine/Plugins/Experimental/UAF/UAFAnimGraph/Source/UAFAnimGraph/Public/TraitInterfaces/IUpdate.h
--- a/Engine/Plugins/Experimental/UAF/UAFAnimGraph/Source/UAFAnimGraph/Public/TraitInterfaces/IUpdate.h
+++ b/Engine/Plugins/Experimental/UAF/UAFAnimGraph/Source/UAFAnimGraph/Public/TraitInterfaces/IUpdate.h
@@ -58,14 +58,12 @@ namespace UE::UAF
 		float GetDeltaTime() const { return DeltaTime; }
 
 	protected:
-		const TWeakObjectPtr<const USkeletalMeshComponent>& GetBindingObject() const { return BindingObject; }
-		void SetBindingObject(const TWeakObjectPtr<const USkeletalMeshComponent>& InBindingObject) { BindingObject = InBindingObject; }
+		TWeakObjectPtr<const USkeletalMeshComponent> GetBindingObject() const;
 
 		FAnimNextGraphInstance& GraphInstance;
 		FTraitEventList InputEventList;
 		FTraitEventList OutputEventList;
 		float DeltaTime = 0.f;
-		TWeakObjectPtr<const USkeletalMeshComponent> BindingObject = nullptr;
 
 		friend UAFANIMGRAPH_API void UpdateGraph(FUpdateGraphContext& UpdateGraphContext);
 		friend struct ::FRigUnit_AnimNextRunAnimationGraph_v1;
diff --git a/Engine/Plugins/Experimental/UAF/UAFAnimGraph/Source/UAFAnimGraphUncookedOnly/Private/Templates/UAFGraphNodeTemplate_InputPose.h b/Engine/Plugins/Experimental/UAF/UAFAnimGraph/Source/UAFAnimGraphUncookedOnly/Private/Templates/UAFGraphNodeTemplate_InputPose.h
new file mode 100644
--- /dev/null
+++ b/Engine/Plugins/Experimental/UAF/UAFAnimGraph/Source/UAFAnimGraphUncookedOnly/Private/Templates/UAFGraphNodeTemplate_InputPose.h
@@ -0,0 +1,81 @@
+// Copyright Epic Games, Inc. All Rights Reserved.
+
+#pragma once
+
+#include "AnimGraphUncookedOnlyUtils.h"
+#include "AnimNextController.h"
+#include "AssetDefinition.h"
+#include "Templates/UAFGraphNodeTemplate.h"
+#include "Traits/InputValueTrait.h"
+#include "UAF/ValueRuntime/ValueBundle.h"
+#include "UAFGraphNodeTemplate_InputPose.generated.h"
+
+#define LOCTEXT_NAMESPACE "UAFGraphNodeTemplate_InputValue"
+
+UCLASS()
+class UUAFGraphNodeTemplate_InputValue : public UUAFGraphNodeTemplate
+{
+	GENERATED_BODY()
+
+	UUAFGraphNodeTemplate_InputValue()
+	{
+		Title = LOCTEXT("InputValueTitle", "Input Value");
+		TooltipText = LOCTEXT("InputValueTooltip", "An input value bundle from an external source");
+		Category = LOCTEXT("InputValueCategory", "UAF");
+		MenuDescription = LOCTEXT("InputValueMenuDesc", "Input Value");
+		Color = FLinearColor(FColor(80, 10, 123));
+		Traits =
+		{
+			TInstancedStruct<FUAFInputValueTraitSharedData>::Make()
+		};
+		DragDropVariableTypes.Add(FAnimNextParamType::GetType<FUAFValueBundle>());
+		SetCategoryForPinsInLayout(
+			{
+				GET_PIN_PATH_STRING_CHECKED(FUAFInputValueTraitSharedData, Input),
+			},
+			FRigVMPinCategory::GetDefaultCategoryName(),
+			NodeLayout,
+			true);
+	}
+
+	virtual void HandleVariableDropped_Implementation(UAnimNextController* Controller, URigVMUnitNode* Node, FAnimNextVariableReference Variable) const
+	{
+		if (Variable.IsNone())
+		{
+			return;
+		}
+
+		Controller->OpenUndoBracket(LOCTEXT("ConfigureNodeOnDrop", "Configure Node On Drop").ToString());
+
+		URigVMPin* Pin = Node->FindPin(TEXT("UAFInputValueTraitSharedData.Input"));
+		if (ensure(Pin))
+		{
+			FString DefaultValue;
+			FAnimNextVariableReference::StaticStruct()->ExportText(DefaultValue, &Variable, nullptr, nullptr, PPF_None, nullptr);
+			Controller->SetPinDefaultValue(Pin, DefaultValue, true, true, true, true);
+			
+			Controller->SetNodeTitle(Node, FText::Format(LOCTEXT("NodeTitleFormat", "Input Value: {0}"), FText::FromName(Variable.GetName())).ToString(), true, true, true);
+		}
+
+		Controller->CloseUndoBracket();
+	}
+
+	virtual void HandlePinDefaultValueChanged_Implementation(UAnimNextController* Controller, URigVMPin* Pin) const
+	{
+		Super::HandlePinDefaultValueChanged_Implementation(Controller, Pin);
+
+		if (Pin->GetFName() == GET_MEMBER_NAME_CHECKED(FUAFInputValueTraitSharedData, Input))
+		{
+			Controller->OpenUndoBracket(LOCTEXT("SetNodeTitle", "Set Node Title").ToString());
+
+			FString DefaultValue = Pin->GetDefaultValue();
+			FAnimNextVariableReference VariableReference;
+			FAnimNextVariableReference::StaticStruct()->ImportText(*DefaultValue, &VariableReference, nullptr, PPF_None, nullptr, FAnimNextVariableReference::StaticStruct()->GetName());
+			Controller->SetNodeTitle(Pin->GetNode(), FText::Format(LOCTEXT("NodeTitleFormat", "Input Value: {0}"), FText::FromName(VariableReference.GetName())).ToString(), true, true, true);
+
+			Controller->CloseUndoBracket();
+		}
+	}
+};
+
+#undef LOCTEXT_NAMESPACE
\ No newline at end of file
diff --git a/Engine/Plugins/Experimental/UAF/UAFTestSuites/Source/UAFCQTestSuite/Private/UAFModuleTests.cpp b/Engine/Plugins/Experimental/UAF/UAFTestSuites/Source/UAFCQTestSuite/Private/UAFModuleTests.cpp
--- a/Engine/Plugins/Experimental/UAF/UAFTestSuites/Source/UAFCQTestSuite/Private/UAFModuleTests.cpp
+++ b/Engine/Plugins/Experimental/UAF/UAFTestSuites/Source/UAFCQTestSuite/Private/UAFModuleTests.cpp
@@ -19,9 +19,25 @@
 #include "AnimNextRigVMAssetEditorData.h"
 #include "Math/Vector2D.h"
 #include "AssetToolsModule.h"
+#include "BlueprintEditor.h"
+#include "BlueprintEditorModule.h"
 #include "AssetRegistry/AssetRegistryModule.h"
 #include "ObjectTools.h"
+#include "Animation/AnimSequence.h"
+#include "Asset/UAFAnimGraphAssetData.h"
+#include "Component/AnimNextComponent.h"
+#include "Engine/SCS_Node.h"
+#include "Engine/SimpleConstructionScript.h"
+#include "Engine/SkeletalMesh.h"
+#include "Factories/BlueprintFactory.h"
+#include "GameFramework/Character.h"
+#include "Graph/RigUnit_UAFRunAsset.h"
+#include "Graph/RigUnit_UAFWriteSystemOutput.h"
+#include "Graph/UAFSystemOutputComponent.h"
+#include "Kismet2/KismetEditorUtilities.h"
+#include "Module/AnimNextModule_EditorData.h"
 #include "Module/UAFSystemAssetData.h"
+#include "Subsystems/EditorActorSubsystem.h"
 #include "UObject/SavePackage.h"
 
 TEST_CLASS(FAnimNextModuleTests, "Animation.UAF.Functional")
@@ -276,6 +292,70 @@ TEST_CLASS(FAnimNextModuleTests, "Animation.UAF.Functional")
 			});
 	}
 	
+	TEST_METHOD(System_Input_Output)
+	{
+		TestCommandBuilder
+			.Do(TEXT("Create UAF System Asset"), [&]()
+			{
+				CreateAnimNextModule();
+			})
+			.Then(TEXT("Create a Run Graph node"), [&]()
+			{
+				EditorData = UE::UAF::UncookedOnly::FUtils::GetEditorData(AnimNextModule);
+				ASSERT_THAT(IsNotNull(EditorData));
+				ControllerGraph = EditorData->GetControllerByName("RigVMGraph")->GetGraph();
+				ASSERT_THAT(IsNotNull(ControllerGraph));				
+				ParentGraph = Cast<UEdGraph>(EditorData->GetEditorObjectForRigVMGraph(ControllerGraph));
+				ASSERT_THAT(IsNotNull(ParentGraph));	
+				
+				TArray<UEdGraphPin*> FromPins;
+				RunGraphNode = UAFTestsUtilities::AddUnitNode(ParentGraph, FRigUnit_AnimNextRunAnimationGraph_v2::StaticStruct()->GetPathName(), FromPins, FVector2f(64.0f, 16.0f));
+				ASSERT_THAT(IsNotNull(RunGraphNode));
+			})
+			.Then(TEXT("Create a Write Pose node"), [&]()
+			{
+				TArray<UEdGraphPin*> FromPins;
+				WritePoseNode = UAFTestsUtilities::AddUnitNode(ParentGraph, FRigUnit_AnimNextWriteSkeletalMeshComponentPose::StaticStruct()->GetPathName(), FromPins, FVector2f(336.0f, 0.0f));
+				ASSERT_THAT(IsNotNull(WritePoseNode));
+			})
+			.Then(TEXT("Add link from PrePhysics node Execute pin to Run Graph node Execute pin"), [&]()
+			{
+				ASSERT_THAT(IsTrue(UAFTestsUtilities::AddLink(AnimNextModule, "RigUnit_AnimNextPrePhysicsEvent.ExecuteContext", "AnimNextRunAnimationGraph_v2.ExecuteContext")));
+			})
+			.Then(TEXT("Add link from Run Graph node output pin to Write Pose node Execute pin"), [&]()
+			{				
+				ASSERT_THAT(IsTrue(UAFTestsUtilities::AddLink(AnimNextModule, "AnimNextRunAnimationGraph_v2.ExecuteContext", "AnimNextWriteSkeletalMeshComponentPose.ExecuteContext")));
+			})
+			.Then(TEXT("Collpase Run Graph and Write Pose nodes"), [&]()
+			{
+				TArray<FName> Nodes;
+				Nodes.Add(FName(RunGraphNodeName));
+				Nodes.Add(FName(WritePoseNodeName));
+				URigVMCollapseNode* CollapseNode = UAFTestsUtilities::CollapseNodes(AnimNextModule, Nodes);
+				ASSERT_THAT(IsNotNull(CollapseNode));
+			})
+			.Until(TEXT("Execute Undo Collpase Nodes"), [&]()
+			{
+				return GEditor->UndoTransaction();
+			})
+			.Then(TEXT("Assert Undo Collapse Nodes"), [&]()
+			{
+				ASSERT_THAT(IsNotNull(ControllerGraph->FindNodeByName(FName(RunGraphNodeName))));
+				ASSERT_THAT(IsNotNull(ControllerGraph->FindNodeByName(FName(WritePoseNodeName))));
+				ASSERT_THAT(IsNull(ControllerGraph->FindNodeByName(FName(CollapseNodeName))));
+			})
+			.Until(TEXT("Execute Redo Collapse Nodes"), [&]()
+			{
+				return GEditor->RedoTransaction();
+			})
+			.Then(TEXT("Assert Redo Collapse Nodes"), [&]()
+			{
+				ASSERT_THAT(IsNull(ControllerGraph->FindNodeByName(FName(RunGraphNodeName))));
+				ASSERT_THAT(IsNull(ControllerGraph->FindNodeByName(FName(WritePoseNodeName))));
+				ASSERT_THAT(IsNotNull(ControllerGraph->FindNodeByName(FName(CollapseNodeName))));
+			});
+	}
+	
 protected:
 	void CreateAnimNextModule()
 	{
@@ -302,6 +382,236 @@ protected:
 namespace UE::UAF::Tests
 {
 
+TEST_CLASS(FUAFSystemTests, "Animation.UAF.System")
+{
+	struct FSystemTestData
+	{
+		TStrongObjectPtr<UUAFSystem> System;
+		UUAFSystem_EditorData* EditorData = nullptr;
+		URigVMGraph* ControllerGraph = nullptr;
+		UEdGraph* ParentGraph = nullptr;
+		TStrongObjectPtr<UUAFComponent> Component;
+	};
+
+	TArray<FSystemTestData> SystemData;
+	UEditorActorSubsystem* EditorActorSubsystem = nullptr;
+	TStrongObjectPtr<USkeleton> Skeleton;
+	TStrongObjectPtr<USkeletalMesh> SkeletalMesh;
+	TStrongObjectPtr<UAnimSequence> AnimSequence;
+	TStrongObjectPtr<UBlueprint> Blueprint;
+	AActor* Actor = nullptr;
+
+	BEFORE_EACH()
+	{
+		EditorActorSubsystem = GEditor->GetEditorSubsystem<UEditorActorSubsystem>();
+		ASSERT_THAT(IsNotNull(EditorActorSubsystem));
+
+		UPackage* SkeletonPackage = CreateUniquePackage("NewSkeletalMesh");
+		Skeleton = TStrongObjectPtr(NewObject<USkeleton>(SkeletonPackage, FName(), RF_Public | RF_Standalone | RF_Transient));
+		{
+			FReferenceSkeletonModifier SkeletonModifier(Skeleton.Get());
+
+			// Add a root bone
+			SkeletonModifier.Add(FMeshBoneInfo(NAME_Root, FString(TEXT("Root")), -1), FTransform::Identity);
+		}
+
+		UPackage* SkeletalMeshPackage = CreateUniquePackage("NewSkeletalMesh");
+		SkeletalMesh = TStrongObjectPtr(NewObject<USkeletalMesh>(SkeletalMeshPackage, FName(), RF_Public | RF_Standalone | RF_Transient));
+		SkeletalMesh->SetSkeleton(Skeleton.Get());
+
+		UPackage* AnimSequencePackage = CreateUniquePackage("NewAnimSequence");
+		AnimSequence = TStrongObjectPtr(NewObject<UAnimSequence>(AnimSequencePackage, FName(), RF_Public | RF_Standalone | RF_Transient));
+		AnimSequence->SetSkeleton(Skeleton.Get());
+
+		const int32 TestNumFrames = 30;
+		const FFrameRate TestFps = FFrameRate(30, 1);
+
+		TScriptInterface<IAnimationDataController> Controller = AnimSequence->GetDataModel()->GetController();
+		Controller->OpenBracket(NSLOCTEXT("UAFSystemTests", "CreateAnimSequence", "Creating Animation Sequence"), false);
+		Controller->InitializeModel();
+		Controller->SetFrameRate(TestFps, false);
+		Controller->SetNumberOfFrames(TestNumFrames, false);
+		Controller->CloseBracket(false);
+		AnimSequence->WaitOnExistingCompression();
+	}
+
+	AFTER_EACH()
+	{
+		for(const FSystemTestData& Data : SystemData)
+		{
+			if (Data.System)
+			{
+				ObjectTools::DeleteSingleObject(Data.System.Get());
+			}
+		}
+		SystemData.Empty();
+		EditorActorSubsystem->DestroyActor(Actor);
+		EditorActorSubsystem = nullptr;
+		Skeleton = nullptr;
+		SkeletalMesh = nullptr;
+		Blueprint = nullptr;
+		Actor = nullptr;
+		
+		FUtils::CleanupAfterTests();
+	}
+
+	TEST_METHOD(System_Input_Output)
+	{
+		TestCommandBuilder
+			.Do(TEXT("Create UAF System Assets"), [&]()
+			{
+				// First system uses an asset, so we dont need to m,ess with its graph
+				{
+					FSystemTestData& System0 = SystemData.AddDefaulted_GetRef();
+				}
+
+				// Second system writes input to output
+				{
+					FSystemTestData& System1 = CreateSystem();
+					TArray<UEdGraphPin*> FromPins;
+
+					UAnimNextVariableEntry* VariableEntry = UAFTestsUtilities::AddVariable(System1.System.Get(), FAnimNextParamType::GetType<FUAFValueBundle>(), TEXT("InputValues"), TEXT(""));
+					ASSERT_THAT(IsNotNull(VariableEntry));
+
+					UEdGraphNode* VariableNode = UAFTestsUtilities::AddVariableNode(System1.ParentGraph, System1.System.Get(), TEXT("InputValues"), FAnimNextParamType::GetType<FUAFValueBundle>(), FAnimNextSchemaAction_Variable::EVariableAccessorChoice::Get, FromPins, FVector2f(64.0f, 16.0f));
+					ASSERT_THAT(IsNotNull(VariableNode));
+
+					UEdGraphNode* WritePoseNode = UAFTestsUtilities::AddUnitNode(System1.ParentGraph, FRigUnit_UAFWriteSystemOutput::StaticStruct()->GetPathName(), FromPins, FVector2f(336.0f, 0.0f));
+					ASSERT_THAT(IsNotNull(WritePoseNode));
+
+					// Link PrePhysics to Output
+					ASSERT_THAT(IsTrue(UAFTestsUtilities::AddLink(System1.System.Get(), "RigUnit_AnimNextPrePhysicsEvent.ExecuteContext", "UAFWriteSystemOutput.ExecuteContext")));
+
+					// Link input var to output
+					ASSERT_THAT(IsTrue(UAFTestsUtilities::AddLink(System1.System.Get(), "VariableNode.Value", "UAFWriteSystemOutput.Value")));
+				}
+			})
+			.Do(TEXT("Create Character Blueprint Asset"), [&]()
+			{
+				UBlueprintFactory* Factory = NewObject<UBlueprintFactory>();
+				Factory->ParentClass = ACharacter::StaticClass();
+		
+				const FString BlueprintName = "NewBlueprint";
+				UObject* FactoryObject = UAFTestsUtilities::CreateFactoryObject(Factory, UBlueprint::StaticClass(), BlueprintName);
+				Blueprint = TStrongObjectPtr(CastChecked<UBlueprint>(FactoryObject));
+				ASSERT_THAT(IsNotNull(Blueprint));
+				ASSERT_THAT(AreEqual(BlueprintName, Blueprint->GetFName()));
+			})
+			.Then(TEXT("Create Blueprint"), [&]()
+			{
+				TArray<UActorComponent*> Components;
+				{
+					FSystemTestData& System0 = SystemData[0];
+					const FName ComponentName = FName(TEXT("UAF0"));
+					System0.Component = TStrongObjectPtr(NewObject<UUAFComponent>(GetTransientPackage(), UUAFComponent::StaticClass(), ComponentName));
+					ASSERT_THAT(IsNotNull(System0.Component));
+					ASSERT_THAT(AreEqual(ComponentName, System0.Component->GetFName()));
+
+					// First component just plays a simple animation
+					System0.Component->AssetData = TInstancedStruct<FUAFGraphFactoryAsset_Animation>::Make(AnimSequence.Get());
+
+					Components.Add(System0.Component.Get());
+				}
+
+				{
+					FSystemTestData& System1 = SystemData[1];
+					const FName ComponentName = FName(TEXT("UAF1"));
+					System1.Component = TStrongObjectPtr(NewObject<UUAFComponent>(GetTransientPackage(), UUAFComponent::StaticClass(), ComponentName));
+					ASSERT_THAT(IsNotNull(System1.Component));
+					ASSERT_THAT(AreEqual(ComponentName, System1.Component->GetFName()));
+
+					// Tag component so we can find it later
+					System1.Component->ComponentTags.Add(TEXT("Output"));
+
+					// Second component uses our system
+					System1.Component->AssetData = TInstancedStruct<FUAFSystemFactoryAsset_System>::Make(System1.System.Get());
+
+					// Set UAF0 as an input to this component
+					FUAFComponentInputDesc InputDesc;
+					InputDesc.Component.PathToComponent = TEXT("UAF0");
+					InputDesc.Input = FAnimNextVariableReference::FromName(TEXT("InputValues"), System1.System.Get());
+					System1.Component->Inputs.Add(InputDesc);
+
+					// Set output to the character's mesh component
+					System1.Component->OutputComponent.PathToComponent = TEXT("Mesh");
+					Components.Add(System1.Component.Get());
+				}
+
+				FKismetEditorUtilities::AddComponentsToBlueprint(Blueprint.Get(), Components);
+			})
+			.Then(TEXT("Compile Blueprint"), [&]()
+			{
+				auto CompileBlueprint = [this]()
+				{
+					FCompilerResultsLog LogResults;
+					LogResults.SetSourcePath(Blueprint->GetPathName());
+					FKismetEditorUtilities::CompileBlueprint(Blueprint.Get(), EBlueprintCompileOptions::None, &LogResults);
+				};
+
+				// Perform initial compile to generate the BP CDO
+				CompileBlueprint();
+				ASSERT_THAT(IsTrue(Blueprint->Status == EBlueprintStatus::BS_UpToDate));
+				ASSERT_THAT(IsNotNull(Blueprint->GeneratedClass));
+
+				// Modify the mesh component to use the correct mesh and recompile
+				ACharacter* CDO = Cast<ACharacter>(Blueprint->GeneratedClass->GetDefaultObject());
+				ASSERT_THAT(IsNotNull(CDO));
+				USkeletalMeshComponent* MeshComponent = CDO->GetMesh();
+				ASSERT_THAT(IsNotNull(MeshComponent));
+				MeshComponent->SetEnableAnimation(false);
+				MeshComponent->SetSkeletalMesh(SkeletalMesh.Get(), false);
+				CompileBlueprint();
+				ASSERT_THAT(IsTrue(Blueprint->Status == EBlueprintStatus::BS_UpToDate));
+			})
+			.Then(TEXT("Spawn BP actor in editor level"), [&]()
+			{
+				// Our test mesh is not really valid for use, so suppress some warnings here
+				TestRunner->AddExpectedError(TEXT("Could not write to skeletal mesh component"), EAutomationExpectedErrorFlags::Contains, 4);
+				TestRunner->AddExpectedError(TEXT("Error generating CanonicalBoneSet for SkeletalMesh"));
+				TestRunner->AddExpectedError(TEXT("Skeletal Mesh asset '")); // SkeletalMesh_0' has no render data
+
+				Actor = EditorActorSubsystem->SpawnActorFromClass(Blueprint->GeneratedClass.Get(), FVector::ZeroVector, FRotator::ZeroRotator, true);
+				ASSERT_THAT(IsNotNull(Actor));
+				UUAFComponent* OutputComponent = Actor->FindComponentByTag<UUAFComponent>(TEXT("Output"));
+				ASSERT_THAT(IsNotNull(OutputComponent));
+				uint32 SerialNumber = 0;
+				OutputComponent->GetSystemReference().ReadComponent<FUAFSystemOutputComponent>([&](TConstStructView<FUAFSystemOutputComponent> InSystemOutputComponent)
+				{
+					SerialNumber = InSystemOutputComponent.Get().SerialNumber;
+				});
+				ASSERT_THAT(IsTrue(SerialNumber != 0));
+			});
+	}
+
+protected:
+	FSystemTestData& CreateSystem()
+	{
+		UFactory* Factory = NewObject<UUAFSystemFactory>();
+		UPackage* SystemPackage = CreateUniquePackage("NewUAFSystem");
+		FName SystemPackageName = *FPaths::GetBaseFilename(SystemPackage->GetName());
+		UObject* FactoryObject = Factory->FactoryCreateNew(UUAFSystem::StaticClass(), SystemPackage, SystemPackageName, RF_Public | RF_Standalone, NULL, GWarn);
+		FAssetRegistryModule::AssetCreated(FactoryObject);
+
+		FSystemTestData& TestData = SystemData.Emplace_GetRef();
+		TestData.System = TStrongObjectPtr(CastChecked<UUAFSystem>(FactoryObject));
+		TestData.EditorData = CastChecked<UUAFSystem_EditorData>(UE::UAF::UncookedOnly::FUtils::GetEditorData(TestData.System.Get()));
+		TestData.ControllerGraph = TestData.EditorData->GetControllerByName("RigVMGraph")->GetGraph();
+		TestData.ParentGraph = Cast<UEdGraph>(TestData.EditorData->GetEditorObjectForRigVMGraph(TestData.ControllerGraph));
+		return TestData;
+	}
+
+	UPackage* CreateUniquePackage(FString InAssetName)
+	{
+		// Create a unique package name and path
+		const FString BasePackageName = "/Game/";    // aka “/All/Content/‘
+		FString UniquePackageName;
+		FString AssetName;
+		IAssetTools& AssetTools = FModuleManager::LoadModuleChecked<FAssetToolsModule>("AssetTools").Get();
+		AssetTools.CreateUniqueAssetName(BasePackageName, InAssetName, UniquePackageName, AssetName);
+		return CreatePackage(*UniquePackageName);
+	}
+};
+
 IMPLEMENT_SIMPLE_AUTOMATION_TEST(FSystemReference_Creation, "Animation.UAF.Runtime.SystemReference.Creation", EAutomationTestFlags::EditorContext | EAutomationTestFlags::EngineFilter)
 
 bool FSystemReference_Creation::RunTest(const FString& InParameters)
