Add typed versions of SetValue/GetValue to TBoundValueMap

#Synced-CL 51636272

diff --git a/Engine/Plugins/Experimental/UAF/UAF/Source/UAF/Public/UAF/ValueRuntime/BoundValueMap.h b/Engine/Plugins/Experimental/UAF/UAF/Source/UAF/Public/UAF/ValueRuntime/BoundValueMap.h
--- a/Engine/Plugins/Experimental/UAF/UAF/Source/UAF/Public/UAF/ValueRuntime/BoundValueMap.h
+++ b/Engine/Plugins/Experimental/UAF/UAF/Source/UAF/Public/UAF/ValueRuntime/BoundValueMap.h
@@ -165,6 +165,18 @@ namespace UE::UAF
 		// Returns the total size required for an instance of a map of this type
 		[[nodiscard]] static int32 CalculateAllocationSize(int32 NumAttributes);
 
+		// Typed equivalent for FBoundValueMap::SetValue
+		void SetValue(FAttributeSetIndex Index, const InValueType& Value)
+		{
+			FBoundValueMap::SetValue<InValueType>(Index, Value);
+		}
+
+		// Typed equivalent for FBoundValueMap::GetValue
+		[[nodiscard]] InValueType GetValue(FAttributeSetIndex Index) const
+		{
+			return FBoundValueMap::GetValue<InValueType>(Index);
+		}
+		
 	protected:
 		// Constructs a map
 		TBoundValueMap(const FConstructArgs& Args, int32 AllocatedSize);
diff --git a/Engine/Plugins/Experimental/UAF/UAFTestSuites/Source/UAFTestSuite/Private/UAF/ValueRuntime/BoundValueMapTest.cpp b/Engine/Plugins/Experimental/UAF/UAFTestSuites/Source/UAFTestSuite/Private/UAF/ValueRuntime/BoundValueMapTest.cpp
--- a/Engine/Plugins/Experimental/UAF/UAFTestSuites/Source/UAFTestSuite/Private/UAF/ValueRuntime/BoundValueMapTest.cpp
+++ b/Engine/Plugins/Experimental/UAF/UAFTestSuites/Source/UAFTestSuite/Private/UAF/ValueRuntime/BoundValueMapTest.cpp
@@ -112,7 +112,7 @@ namespace UE::UAF::Tests
 				FFloatAnimationAttribute Value2{ -0.13544f };
 				BoneToFloatMap.SetValue(FAttributeSetIndex(2), Value2);
 
-				AddErrorIfFalse(BoneToFloatMap.GetValue<FFloatAnimationAttribute>(FAttributeSetIndex(2)).Value == Value2.Value, TEXT("Unexpected mapped value"));
+				AddErrorIfFalse(BoneToFloatMap.GetValue(FAttributeSetIndex(2)).Value == Value2.Value, TEXT("Unexpected mapped value"));
 
 				// Copy
 				{
@@ -167,7 +167,7 @@ namespace UE::UAF::Tests
 				FBoneTransformAnimationAttribute Value2{ FTransform(FVector(0.233, 331.54, 571.146)) };
 				BoneToTransformMap.SetValue(FAttributeSetIndex(2), Value2);
 
-				AddErrorIfFalse(BoneToTransformMap.GetValue<FBoneTransformAnimationAttribute>(FAttributeSetIndex(2)).Value.Equals(Value2.Value), TEXT("Unexpected mapped value"));
+				AddErrorIfFalse(BoneToTransformMap.GetValue(FAttributeSetIndex(2)).Value.Equals(Value2.Value), TEXT("Unexpected mapped value"));
 
 				// Copy
 				{
@@ -222,7 +222,7 @@ namespace UE::UAF::Tests
 				FStringAnimationAttribute Value2{ TEXT("01Wow") };
 				BoneToStringMap.SetValue(FAttributeSetIndex(2), Value2);
 
-				AddErrorIfFalse(BoneToStringMap.GetValue<FStringAnimationAttribute>(FAttributeSetIndex(2)).Value.Equals(Value2.Value), TEXT("Unexpected mapped value"));
+				AddErrorIfFalse(BoneToStringMap.GetValue(FAttributeSetIndex(2)).Value.Equals(Value2.Value), TEXT("Unexpected mapped value"));
 
 				// Copy
 				{
