UAF: Fix uninitialized struct member warning

#Synced-CL 51880276

diff --git a/Engine/Plugins/Experimental/UAF/UAFTestSuites/Source/UAFAnimGraphTestSuite/Private/AnimNextInlineSubGraphTraitTest.h b/Engine/Plugins/Experimental/UAF/UAFTestSuites/Source/UAFAnimGraphTestSuite/Private/AnimNextInlineSubGraphTraitTest.h
--- a/Engine/Plugins/Experimental/UAF/UAFTestSuites/Source/UAFAnimGraphTestSuite/Private/AnimNextInlineSubGraphTraitTest.h
+++ b/Engine/Plugins/Experimental/UAF/UAFTestSuites/Source/UAFAnimGraphTestSuite/Private/AnimNextInlineSubGraphTraitTest.h
@@ -12,5 +12,5 @@ struct FInlineSubGraphTest_LeafTraitSharedData : public FAnimNextTraitSharedData
 	GENERATED_BODY()
 
 	UPROPERTY()
-	int32 Value;
+	int32 Value = 0;
 };
