UAF Transformer Tests

#Synced-CL 52003153

diff --git a/Engine/Plugins/Experimental/UAF/UAFTestSuites/Source/UAFTestSuite/Private/UAF/ValueRuntime/Transformers/BoneTransformAttributeTransformerTests.cpp b/Engine/Plugins/Experimental/UAF/UAFTestSuites/Source/UAFTestSuite/Private/UAF/ValueRuntime/Transformers/BoneTransformAttributeTransformerTests.cpp
new file mode 100644
--- /dev/null
+++ b/Engine/Plugins/Experimental/UAF/UAFTestSuites/Source/UAFTestSuite/Private/UAF/ValueRuntime/Transformers/BoneTransformAttributeTransformerTests.cpp
@@ -0,0 +1,330 @@
+// Copyright Epic Games, Inc. All Rights Reserved.
+
+#pragma once
+
+#include "CoreMinimal.h"
+#include "Misc/AutomationTest.h"
+
+#if WITH_DEV_AUTOMATION_TESTS
+
+#include "AnimNextTest.h"
+#include "Animation/BuiltInAttributeTypes.h"
+#include "UAF/AbstractSkeleton/AbstractSkeletonSetBinding.h"
+#include "UAF/Attributes/AttributeBindingDataCache.h"
+#include "UAF/Attributes/EngineAttributes.h"
+#include "UAF/ValueRuntime/ValueRuntimeRegistry.h"
+#include "UAF/ValueRuntime/Transformers/Accumulate.h"
+#include "UAF/ValueRuntime/Transformers/AdditiveSpace.h"
+#include "UAF/ValueRuntime/Transformers/Interpolate.h"
+#include "UAF/ValueRuntime/Transformers/Overwrite.h"
+#include "UAF/ValueRuntime/Transformers/Sanitize.h"
+#include "UAF/ValueRuntime/Transformers/TransformerTestUtil.h"
+#include "UObject/Package.h"
+
+namespace UE::UAF::Tests
+{
+	IMPLEMENT_SIMPLE_AUTOMATION_TEST(FTransformerTests_BoneTransformAttribute, "Animation.UAF.ValueRuntime.Transformers.BoneTransformAttribute", EAutomationTestFlags::EditorContext | EAutomationTestFlags::EngineFilter)
+
+	bool FTransformerTests_BoneTransformAttribute::RunTest(const FString& InParameters)
+	{
+		ON_SCOPE_EXIT{ FUtils::CleanupAfterTests(); };
+
+		 TValueOrError<TNonNullPtr<UAbstractSkeletonSetBinding>, FString> SetBindingResult = BuildSetBinding([](TNonNullPtr<USkeleton> Skeleton)
+			{
+				// Not terribly clean, we cast away the 'const' to modify the skeleton
+				FReferenceSkeleton& RefSkeleton = const_cast<FReferenceSkeleton&>(Skeleton->GetReferenceSkeleton());
+				FReferenceSkeletonModifier SkeletonModifier(RefSkeleton, Skeleton);
+
+				SkeletonModifier.Add(FMeshBoneInfo(TEXT("AttrA"), TEXT("AttrA"), INDEX_NONE), FTransform::Identity);
+				SkeletonModifier.Add(FMeshBoneInfo(TEXT("AttrB"), TEXT("AttrB"), 0), FTransform::Identity);
+				SkeletonModifier.Add(FMeshBoneInfo(TEXT("AttrC"), TEXT("AttrC"), 0), FTransform::Identity);
+
+				// When our modifier is destroyed here, it will rebuild the skeleton
+			});
+
+		if (SetBindingResult.HasError())
+		{
+			AddError(SetBindingResult.GetError());
+			return false;
+		}
+		
+		TNonNullPtr<UAbstractSkeletonSetBinding> SetBinding = SetBindingResult.GetValue();
+		
+		SetBinding->AddAttributeToSet(FAnimationAttributeIdentifier(TEXT("AttrA"), INDEX_NONE, NAME_None, FBoneTransformAnimationAttribute::StaticStruct()), "SetA");
+		SetBinding->AddAttributeToSet(FAnimationAttributeIdentifier(TEXT("AttrB"), INDEX_NONE, NAME_None, FBoneTransformAnimationAttribute::StaticStruct()), "SetB");
+		SetBinding->AddAttributeToSet(FAnimationAttributeIdentifier(TEXT("AttrC"), INDEX_NONE, NAME_None, FBoneTransformAnimationAttribute::StaticStruct()), "SetC");
+
+		const FAttributeBindingDataPtr& BindingData = GAttributeBindingDataCache.GetOrAdd(SetBinding, nullptr);
+		FAttributeNamedSetPtr SetA = BindingData->FindNamedSet("SetA");
+		FAttributeNamedSetPtr SetB = BindingData->FindNamedSet("SetB");
+		FAttributeNamedSetPtr SetC = BindingData->FindNamedSet("SetC");
+
+		const FTransform TransformValueA(FQuat::Identity, FVector(100.0f, 200.0f, 300.0f), FVector::OneVector);
+		const FTransform TransformValueB(FQuat(0.0f, 1.0f, 0.0f, 0.0f), FVector(400.0f, 500.0f, 600.0f), FVector(2.0f, 2.0f, 2.0f));
+
+		// Overwrite
+		AddInfo("Overwrite");
+		{
+			// 0.25 weight
+			{
+				FTransformerUnaryTest<FBoneTransformAnimationAttribute, false> TransformTest;
+				TransformTest.AttributeName = "AttrA";
+				TransformTest.Input.Value = TransformValueA;
+				TransformTest.ApplyTransformer = [](const FValueBundleHeap& InputValues, FValueBundleHeap& OutputValues)
+					{
+						const auto& TransformerMap = FValueRuntimeRegistry::Get().GetTransformerMap();
+						Transformers::FOverwrite::Apply(TransformerMap, InputValues, 0.25f, OutputValues);
+					};
+				TransformTest.IsExpectedResult = [&](TOptional<FBoneTransformAnimationAttribute> Result) -> TOptional<FString>
+					{
+						if (!Result.IsSet())
+						{
+							return FString("Overwrite - 0.25 weight - expected to produce a value but received none");
+						}
+
+						const FTransform ExpectedValue = TransformValueA * ScalarRegister(0.25f);
+						if (!ExpectedValue.Equals(Result->Value))
+						{
+							return FString::Printf(TEXT("Overwrite - 0.25 weight - expected %s, Received %s"), *ExpectedValue.ToHumanReadableString(), *Result->Value.ToHumanReadableString());
+						}
+
+						return TOptional<FString>();
+					};
+
+				TransformTest.Run(*this, SetA);
+			}
+		}
+
+		// Accumulate
+		AddInfo("Accumulate");
+		{
+			// 0.25 weight
+			{
+				FTransformerBinaryTest<Transformers::FAccumulate, FBoneTransformAnimationAttribute, EInplaceTransformationSupport::Allow> TransformTest;
+				TransformTest.AttributeName = "AttrA";
+				TransformTest.InputA.Value = TransformValueA;
+				TransformTest.InputB.Value = TransformValueB;
+				TransformTest.ApplyTransformer = [](const FValueBundleHeap& InputAValues, const FValueBundleHeap& InputBValues, FValueBundleHeap& OutputValues)
+					{
+						const auto& TransformerMap = FValueRuntimeRegistry::Get().GetTransformerMap();
+						Transformers::FAccumulate::Apply(TransformerMap, InputAValues, InputBValues, 0.25f, OutputValues);
+					};
+				TransformTest.IsExpectedResult = [&](TOptional<FBoneTransformAnimationAttribute> Result) -> TOptional<FString>
+					{
+						if (!Result.IsSet())
+						{
+							return FString("Accumulate - 0.25 weight - expected to produce a value but received none");
+						}
+						
+						const FTransform ExpectedValue = TransformValueA + (TransformValueB * ScalarRegister(0.25f));
+						if (!ExpectedValue.Equals(Result->Value))
+						{
+							return FString::Printf(TEXT("Accumulate - 0.25 weight - expected %s, received %s"), *ExpectedValue.ToHumanReadableString(), *Result->Value.ToHumanReadableString());
+						}
+						
+						return TOptional<FString>();
+					};
+
+				TransformTest.Run(*this, SetA);
+			}
+		}
+
+		// Interpolate
+		AddInfo("Interpolate");
+		{
+			// 0.25 weight
+			{
+				FTransformerBinaryTest<Transformers::FInterpolate, FBoneTransformAnimationAttribute, EInplaceTransformationSupport::Allow> TransformTest;
+				TransformTest.AttributeName = "AttrA";
+				TransformTest.InputA.Value = TransformValueA;
+				TransformTest.InputB.Value = TransformValueB;
+				TransformTest.ApplyTransformer = [](const FValueBundleHeap& InputAValues, const FValueBundleHeap& InputBValues, FValueBundleHeap& OutputValues)
+					{
+						const auto& TransformerMap = FValueRuntimeRegistry::Get().GetTransformerMap();
+						Transformers::FInterpolate::Apply(TransformerMap, InputAValues, InputBValues, 0.25f, OutputValues);
+					};
+				TransformTest.IsExpectedResult = [&](TOptional<FBoneTransformAnimationAttribute> Result) -> TOptional<FString>
+					{
+						if (!Result.IsSet())
+						{
+							return FString("Interpolate - 0.25 weight - expected to produce a value but received none");
+						}
+
+						FTransform ExpectedValue = TransformValueA;
+						ExpectedValue.BlendWith(TransformValueB, 0.25f);
+						if (!ExpectedValue.Equals(Result->Value))
+						{
+							return FString::Printf(TEXT("Interpolate - 0.25 weight - expected %s, received %s"), *ExpectedValue.ToHumanReadableString(), *Result->Value.ToHumanReadableString());
+						}
+						
+						return TOptional<FString>();
+					};
+
+				TransformTest.Run(*this, SetA);
+			}
+		}
+
+		// Layer
+		AddInfo("Layer");
+		{
+			// Value in layer but not in base
+			{
+				// Weight 0.25
+				{
+					FTransformerBinaryTest<Transformers::FLayer, FBoneTransformAnimationAttribute, EInplaceTransformationSupport::Allow> TransformTest;
+					TransformTest.AttributeName = "AttrC";
+					TransformTest.WeightB = 0.25f;
+					TransformTest.InputA.Reset();
+					TransformTest.InputB = FBoneTransformAnimationAttribute({ TransformValueB });
+					TransformTest.IsExpectedResult = [&](TOptional<FBoneTransformAnimationAttribute> Result) -> TOptional<FString>
+						{
+							if (Result.IsSet())
+							{
+								return FString("Layer - 0.25 weight, case 2 - expected to not produce a value but received one");
+							}
+
+							return TOptional<FString>();
+						};
+
+					TransformTest.Run(*this, SetB, SetC);
+				}
+			}
+
+			// Value in both base and layer
+			{
+				// Weight 0.25
+				{
+					FTransformerBinaryTest<Transformers::FLayer, FBoneTransformAnimationAttribute, EInplaceTransformationSupport::Allow> TransformTest;
+					TransformTest.AttributeName = "AttrB";
+					TransformTest.WeightB = 0.25f;
+					TransformTest.InputA = FBoneTransformAnimationAttribute({ TransformValueA });
+					TransformTest.InputB = FBoneTransformAnimationAttribute({ TransformValueB });
+					TransformTest.IsExpectedResult = [&](TOptional<FBoneTransformAnimationAttribute> Result) -> TOptional<FString>
+						{
+							if (!Result.IsSet())
+							{
+								return FString("Layer - 0.25 weight, case 3 - expected to produce a value but received none");
+							}
+
+							FTransform ExpectedValue = TransformValueA;
+							ExpectedValue.BlendWith(TransformValueB, 0.25f);
+							if (!ExpectedValue.Equals(Result->Value))
+							{
+								return FString::Printf(TEXT("Layer - 0.25 weight, case 3 - expected %s, received %s"), *ExpectedValue.ToHumanReadableString(), *Result->Value.ToHumanReadableString());
+							}
+
+							return TOptional<FString>();
+						};
+
+					TransformTest.Run(*this, SetA, SetB);
+				}
+			}
+		}
+
+		// Make Additive Space
+		AddInfo("MakeAdditiveSpace");
+		{
+			{
+				FTransformerBinaryTest<Transformers::FMakeAdditiveSpace, FBoneTransformAnimationAttribute, EInplaceTransformationSupport::Disallow> TransformTest;
+				TransformTest.AttributeName = "AttrA";
+				TransformTest.InputA.Value = TransformValueA;
+				TransformTest.InputB.Value = TransformValueB;
+				TransformTest.ApplyTransformer = [](const FValueBundleHeap& InputAValues, const FValueBundleHeap& InputBValues, FValueBundleHeap& OutputValues)
+					{
+						const auto& TransformerMap = FValueRuntimeRegistry::Get().GetTransformerMap();
+						Transformers::FMakeAdditiveSpace::Apply(TransformerMap, InputAValues, InputBValues, OutputValues);
+					};
+				TransformTest.IsExpectedResult = [&](TOptional<FBoneTransformAnimationAttribute> Result) -> TOptional<FString>
+					{
+						if (!Result.IsSet())
+						{
+							return FString("MakeAdditiveSpace - expected to produce a value but received none");
+						}
+						
+						FTransform ExpectedValue = TransformValueB;
+						FAnimationRuntime::ConvertTransformToAdditive(ExpectedValue, TransformValueA);
+						if (!ExpectedValue.Equals(Result->Value))
+						{
+							return FString::Printf(TEXT("MakeAdditiveSpace - expected %s, received %s"), *ExpectedValue.ToHumanReadableString(), *Result->Value.ToHumanReadableString());
+						}
+						return TOptional<FString>();
+					};
+
+				TransformTest.Run(*this, SetA);
+			}
+		}
+
+		// Apply Additive Space
+		AddInfo("ApplyAdditiveSpace");
+		{
+			// 0.25 weight
+			{
+				FTransformerBinaryTest<Transformers::FMakeAdditiveSpace, FBoneTransformAnimationAttribute, EInplaceTransformationSupport::Allow> TransformTest;
+				TransformTest.AttributeName = "AttrA";
+				TransformTest.InputA.Value = TransformValueA;
+				TransformTest.InputB.Value = TransformValueB;
+				TransformTest.ApplyTransformer = [](const FValueBundleHeap& InputAValues, const FValueBundleHeap& InputBValues, FValueBundleHeap& OutputValues)
+					{
+						const_cast<FValueBundleHeap&>(InputBValues).SetValueSpace(FValueSpace(EValueSpaceType::Local, true));
+
+						const auto& TransformerMap = FValueRuntimeRegistry::Get().GetTransformerMap();
+						Transformers::FApplyAdditiveSpace::Apply(TransformerMap, InputAValues, InputBValues, 0.25f, OutputValues);
+					};
+				TransformTest.IsExpectedResult = [&](TOptional<FBoneTransformAnimationAttribute> Result) -> TOptional<FString>
+					{
+						if (!Result.IsSet())
+						{
+							return FString("ApplyAdditiveSpace - 0.25 weight - expected to produce a value but received none");
+						}
+						
+						FTransformAnimationAttribute TransformAttribute(TransformValueA);
+						TransformAttribute.Accumulate(FTransformAnimationAttribute(TransformValueB), 0.25f, EAdditiveAnimationType::AAT_LocalSpaceBase);
+						if (!TransformAttribute.Value.Equals(Result->Value))
+						{
+							return FString::Printf(TEXT("ApplyAdditiveSpace - 0.25 weight - expected %s, received %s"), *TransformAttribute.Value.ToHumanReadableString(), *Result->Value.ToHumanReadableString());
+						}
+						
+						return TOptional<FString>();
+					};
+
+				TransformTest.Run(*this, SetA);
+			}
+		}
+
+		// Sanitize
+		AddInfo("Sanitize");
+		{
+			FTransformerUnaryTest<FBoneTransformAnimationAttribute, true> BoneTransformTest;
+			BoneTransformTest.AttributeName = "AttrA";
+			BoneTransformTest.Input.Value = FTransform(FQuat(0.2f, 0.5f, 0.3f, 0.1f), FVector(100.0f, 50.0f, 25.0f), FVector(1.0f));
+			BoneTransformTest.ApplyTransformer = [](FValueBundleHeap& InputValues)
+				{
+					const auto& TransformerMap = FValueRuntimeRegistry::Get().GetTransformerMap();
+					Transformers::FSanitize::Apply(TransformerMap, InputValues);
+				};
+			BoneTransformTest.IsExpectedResult = [&](TOptional<FBoneTransformAnimationAttribute> Result) -> TOptional<FString>
+				{
+					if (!Result.IsSet())
+					{
+						return FString("Sanitize - expected to produce a value but received none");
+					}
+
+					FTransform ExpectedValue = FTransform(FQuat(0.2f, 0.5f, 0.3f, 0.1f), FVector(100.0f, 50.0f, 25.0f), FVector(1.0f));
+					ExpectedValue.NormalizeRotation();
+					
+					if (!ExpectedValue.Equals(Result->Value))
+					{
+						return FString::Format(TEXT("Sanitize - expected {0}, Received {1}"), { ExpectedValue.ToString(), Result->Value.ToString() });
+					}
+
+					return TOptional<FString>();
+				};
+
+			BoneTransformTest.Run(*this, SetA);
+		}
+
+		return true;
+	}
+}
+
+#endif // WITH_DEV_AUTOMATION_TESTS
diff --git a/Engine/Plugins/Experimental/UAF/UAFTestSuites/Source/UAFTestSuite/Private/UAF/ValueRuntime/Transformers/FloatAttributeTransformerTests.cpp b/Engine/Plugins/Experimental/UAF/UAFTestSuites/Source/UAFTestSuite/Private/UAF/ValueRuntime/Transformers/FloatAttributeTransformerTests.cpp
new file mode 100644
--- /dev/null
+++ b/Engine/Plugins/Experimental/UAF/UAFTestSuites/Source/UAFTestSuite/Private/UAF/ValueRuntime/Transformers/FloatAttributeTransformerTests.cpp
@@ -0,0 +1,302 @@
+// Copyright Epic Games, Inc. All Rights Reserved.
+
+#pragma once
+
+#include "CoreMinimal.h"
+#include "Misc/AutomationTest.h"
+
+#if WITH_DEV_AUTOMATION_TESTS
+
+#include "AnimNextTest.h"
+#include "Animation/BuiltInAttributeTypes.h"
+#include "UAF/AbstractSkeleton/AbstractSkeletonSetBinding.h"
+#include "UAF/Attributes/AttributeBindingDataCache.h"
+#include "UAF/ValueRuntime/ValueRuntimeRegistry.h"
+#include "UAF/ValueRuntime/Transformers/Accumulate.h"
+#include "UAF/ValueRuntime/Transformers/AdditiveSpace.h"
+#include "UAF/ValueRuntime/Transformers/Interpolate.h"
+#include "UAF/ValueRuntime/Transformers/Overwrite.h"
+#include "UAF/ValueRuntime/Transformers/TransformerTestUtil.h"
+#include "UObject/Package.h"
+
+namespace UE::UAF::Tests
+{
+	IMPLEMENT_SIMPLE_AUTOMATION_TEST(FTransformerTests_FloatAttribute, "Animation.UAF.ValueRuntime.Transformers.FloatAttribute", EAutomationTestFlags::EditorContext | EAutomationTestFlags::EngineFilter)
+
+	bool FTransformerTests_FloatAttribute::RunTest(const FString& InParameters)
+	{
+		ON_SCOPE_EXIT{ FUtils::CleanupAfterTests(); };
+
+		TValueOrError<TNonNullPtr<UAbstractSkeletonSetBinding>, FString> SetBindingResult = BuildSetBinding({});
+
+		if (SetBindingResult.HasError())
+		{
+			AddError(SetBindingResult.GetError());
+			return false;
+		}
+		
+		TNonNullPtr<UAbstractSkeletonSetBinding> SetBinding = SetBindingResult.GetValue();
+
+		SetBinding->AddAttributeToSet(FAnimationAttributeIdentifier(TEXT("AttrA"), INDEX_NONE, NAME_None, FFloatAnimationAttribute::StaticStruct()), "SetA");
+		SetBinding->AddAttributeToSet(FAnimationAttributeIdentifier(TEXT("AttrB"), INDEX_NONE, NAME_None, FFloatAnimationAttribute::StaticStruct()), "SetB");
+		SetBinding->AddAttributeToSet(FAnimationAttributeIdentifier(TEXT("AttrC"), INDEX_NONE, NAME_None, FFloatAnimationAttribute::StaticStruct()), "SetC");
+		
+		const FAttributeBindingDataPtr& BindingData = GAttributeBindingDataCache.GetOrAdd(SetBinding, nullptr);
+		FAttributeNamedSetPtr SetA = BindingData->FindNamedSet("SetA");
+		FAttributeNamedSetPtr SetB = BindingData->FindNamedSet("SetB");
+		FAttributeNamedSetPtr SetC = BindingData->FindNamedSet("SetC");
+
+		// Overwrite
+		AddInfo("Overwrite");
+		{
+			// 0.25 weight
+			{
+				FTransformerUnaryTest<FFloatAnimationAttribute, false> FloatTest;
+				FloatTest.AttributeName = "AttrA";
+				FloatTest.Input.Value = 5.0f;
+				FloatTest.ApplyTransformer = [](const FValueBundleHeap& InputValues, FValueBundleHeap& OutputValues)
+					{
+						const auto& TransformerMap = FValueRuntimeRegistry::Get().GetTransformerMap();
+						Transformers::FOverwrite::Apply(TransformerMap, InputValues, 0.25f, OutputValues);
+					};
+				FloatTest.IsExpectedResult = [&](TOptional<FFloatAnimationAttribute> Result) -> TOptional<FString>
+					{
+						if (!Result.IsSet())
+						{
+							return FString("Overwrite - 0.25 weight - expected to produce a value but received none");
+						}
+
+						constexpr float ExpectedValue = 0.25f * 5.0f;
+						if (!FMath::IsNearlyEqual(Result->Value, ExpectedValue))
+						{
+							return FString::Format(TEXT("Overwrite - 0.25 weight - expected {0}, Received {1}"), { ExpectedValue, Result->Value });
+						}
+
+						return TOptional<FString>();
+					};
+			}
+		}
+		
+		// Accumulate
+		AddInfo("Accumulate");
+		{
+			// 0.25 weight
+			{
+				FTransformerBinaryTest<Transformers::FAccumulate, FFloatAnimationAttribute, EInplaceTransformationSupport::Allow> FloatTest;
+				FloatTest.AttributeName = "AttrA";
+				FloatTest.InputA.Value = 5.0f;
+				FloatTest.InputB.Value = 8.0f;
+				FloatTest.ApplyTransformer = [](const FValueBundleHeap& InputAValues, const FValueBundleHeap& InputBValues, FValueBundleHeap& OutputValues)
+					{
+						const auto& TransformerMap = FValueRuntimeRegistry::Get().GetTransformerMap();
+						Transformers::FAccumulate::Apply(TransformerMap, InputAValues, InputBValues, 0.25f, OutputValues);
+					};
+				FloatTest.IsExpectedResult = [&](TOptional<FFloatAnimationAttribute> Result) -> TOptional<FString>
+					{
+						if (!Result.IsSet())
+						{
+							return FString("Accumulate - 0.25 weight - expected to produce a value but received none");
+						}
+						constexpr float ExpectedValue = 5.0f + (0.25f * 8.0f);
+						if (!FMath::IsNearlyEqual(Result->Value, ExpectedValue))
+						{
+							return FString::Format(TEXT("Accumulate - 0.25 weight - expected {0}, received {1}"), { ExpectedValue, Result->Value });
+						}
+						return TOptional<FString>();
+					};
+
+				FloatTest.Run(*this, SetA);
+			}
+		}
+
+		// Interpolate
+		AddInfo("Interpolate");
+		{
+			// 0.25 weight
+			{
+				FTransformerBinaryTest<Transformers::FInterpolate, FFloatAnimationAttribute, EInplaceTransformationSupport::Allow> FloatTest;
+				FloatTest.AttributeName = "AttrA";
+				FloatTest.InputA.Value = 5.0f;
+				FloatTest.InputB.Value = 8.0f;
+				FloatTest.ApplyTransformer = [](const FValueBundleHeap& InputAValues, const FValueBundleHeap& InputBValues, FValueBundleHeap& OutputValues)
+					{
+						const auto& TransformerMap = FValueRuntimeRegistry::Get().GetTransformerMap();
+						Transformers::FInterpolate::Apply(TransformerMap, InputAValues, InputBValues, 0.25f, OutputValues);
+					};
+				FloatTest.IsExpectedResult = [&](TOptional<FFloatAnimationAttribute> Result) -> TOptional<FString>
+					{
+						if (!Result.IsSet())
+						{
+							return FString("Interpolate - 0.25 weight -expected to produce a value but received none");
+						}
+						constexpr float ExpectedValue = (0.75f * 5.0f) + (0.25f * 8.0f);
+						if (!FMath::IsNearlyEqual(Result->Value, ExpectedValue))
+						{
+							return FString::Format(TEXT("Interpolate - 0.25 weight - expected {0}, received {1}"), { ExpectedValue, Result->Value });
+						}
+						return TOptional<FString>();
+					};
+
+				FloatTest.Run(*this, SetA);
+			}
+		}
+
+		// Layer
+		AddInfo("Layer");
+		{
+			// Value in base but not in layer
+			{
+				// Weight 0.25
+				{
+					FTransformerBinaryTest<Transformers::FLayer, FFloatAnimationAttribute, EInplaceTransformationSupport::Allow> FloatTest;
+					FloatTest.AttributeName = "AttrA";
+					FloatTest.WeightB = 0.25f;
+					FloatTest.InputA = FFloatAnimationAttribute({ 5.0f });
+					FloatTest.InputB.Reset();
+					FloatTest.IsExpectedResult = [&](TOptional<FFloatAnimationAttribute> Result) -> TOptional<FString>
+						{
+							if (!Result.IsSet())
+							{
+								return FString("Layer - 0.25 weight, case 1 - expected to produce a value but received none");
+							}
+							
+							constexpr float ExpectedValue = 5.0f;
+							if (!FMath::IsNearlyEqual(Result->Value, ExpectedValue))
+							{
+								return FString::Format(TEXT("Layer - 0.25 weight, case 1 - expected {0}, received {1}"), { ExpectedValue, Result->Value });
+							}
+							
+							return TOptional<FString>();
+						};
+
+					FloatTest.Run(*this, SetA, SetB);
+				}
+			}
+
+			// Value in layer but not in base
+			{
+				// Weight 0.25
+				{
+					FTransformerBinaryTest<Transformers::FLayer, FFloatAnimationAttribute, EInplaceTransformationSupport::Allow> FloatTest;
+					FloatTest.AttributeName = "AttrC";
+					FloatTest.WeightB = 0.25f;
+					FloatTest.InputA.Reset();
+					FloatTest.InputB = FFloatAnimationAttribute({ 8.0f });
+					FloatTest.IsExpectedResult = [&](TOptional<FFloatAnimationAttribute> Result) -> TOptional<FString>
+						{
+							if (Result.IsSet())
+							{
+								return FString("Layer - 0.25 weight, case 2 - expected to not produce a value but received one");
+							}
+														
+							return TOptional<FString>();
+						};
+
+					FloatTest.Run(*this, SetB, SetC);
+				}
+			}
+			
+			// Value in both base and layer
+			{
+				// Weight 0.25
+				{
+					FTransformerBinaryTest<Transformers::FLayer, FFloatAnimationAttribute, EInplaceTransformationSupport::Allow> FloatTest;
+					FloatTest.AttributeName = "AttrB";
+					FloatTest.WeightB = 0.25f;
+					FloatTest.InputA = FFloatAnimationAttribute({ 5.0f });
+					FloatTest.InputB = FFloatAnimationAttribute({ 8.0f });
+					FloatTest.IsExpectedResult = [&](TOptional<FFloatAnimationAttribute> Result) -> TOptional<FString>
+						{
+							if (!Result.IsSet())
+							{
+								return FString("Layer - 0.25 weight, case 3 - expected to produce a value but received none");
+							}
+							
+							constexpr float ExpectedValue = (0.75f * 5.0f) + (0.25F * 8.0f);
+							if (!FMath::IsNearlyEqual(Result->Value, ExpectedValue))
+							{
+								return FString::Format(TEXT("Layer - 0.25 weight, case 3 - expected {0}, received {1}"), { ExpectedValue, Result->Value });
+							}
+							
+							return TOptional<FString>();
+						};
+
+					FloatTest.Run(*this, SetA, SetB);
+				}
+			}
+		}
+
+		// Make Additive Space
+		AddInfo("MakeAdditiveSpace");
+		{
+			// 0.25 weight
+			{
+				FTransformerBinaryTest<Transformers::FMakeAdditiveSpace, FFloatAnimationAttribute, EInplaceTransformationSupport::Disallow> FloatTest;
+				FloatTest.AttributeName = "AttrA";
+				FloatTest.InputA.Value = 5.0f;
+				FloatTest.InputB.Value = 8.0f;
+				FloatTest.ApplyTransformer = [](const FValueBundleHeap& InputAValues, const FValueBundleHeap& InputBValues, FValueBundleHeap& OutputValues)
+					{
+						const auto& TransformerMap = FValueRuntimeRegistry::Get().GetTransformerMap();
+						Transformers::FMakeAdditiveSpace::Apply(TransformerMap, InputAValues, InputBValues, OutputValues);
+					};
+				FloatTest.IsExpectedResult = [&](TOptional<FFloatAnimationAttribute> Result) -> TOptional<FString>
+					{
+						if (!Result.IsSet())
+						{
+							return FString("MakeAdditiveSpace - expected to produce a value but received none");
+						}
+						constexpr float ExpectedValue = 8.0f - 5.0f;
+						if (!FMath::IsNearlyEqual(Result->Value, ExpectedValue))
+						{
+							return FString::Format(TEXT("MakeAdditiveSpace - expected {0}, received {1}"), { ExpectedValue, Result->Value });
+						}
+						return TOptional<FString>();
+					};
+
+				FloatTest.Run(*this, SetA);
+			}
+		}
+
+		// Apply Additive Space
+		AddInfo("ApplyAdditiveSpace");
+		{
+			// 0.25 weight
+			{
+				FTransformerBinaryTest<Transformers::FMakeAdditiveSpace, FFloatAnimationAttribute, EInplaceTransformationSupport::Allow> FloatTest;
+				FloatTest.AttributeName = "AttrA";
+				FloatTest.InputA.Value = 5.0f;
+				FloatTest.InputB.Value = 8.0f;
+				FloatTest.ApplyTransformer = [](const FValueBundleHeap& InputAValues, const FValueBundleHeap& InputBValues, FValueBundleHeap& OutputValues)
+					{
+						const_cast<FValueBundleHeap&>(InputBValues).SetValueSpace(FValueSpace(EValueSpaceType::Local, true));
+						
+						const auto& TransformerMap = FValueRuntimeRegistry::Get().GetTransformerMap();
+						Transformers::FApplyAdditiveSpace::Apply(TransformerMap, InputAValues, InputBValues, 0.25f, OutputValues);
+					};
+				FloatTest.IsExpectedResult = [&](TOptional<FFloatAnimationAttribute> Result) -> TOptional<FString>
+					{
+						if (!Result.IsSet())
+						{
+							return FString("ApplyAdditiveSpace - 0.25 weight - expected to produce a value but received none");
+						}
+						constexpr float ExpectedValue = 5.0f + (0.25f * 8.0f);
+						if (!FMath::IsNearlyEqual(Result->Value, ExpectedValue))
+						{
+							return FString::Format(TEXT("ApplyAdditiveSpace - 0.25 weight - expected {0}, received {1}"), { ExpectedValue, Result->Value });
+						}
+						return TOptional<FString>();
+					};
+
+				FloatTest.Run(*this, SetA);
+			}
+		}
+		
+		// Sanitize
+		// Nothing to do
+		
+		return true;
+	}
+}
+
+#endif // WITH_DEV_AUTOMATION_TESTS
\ No newline at end of file
diff --git a/Engine/Plugins/Experimental/UAF/UAFTestSuites/Source/UAFTestSuite/Private/UAF/ValueRuntime/Transformers/IntegerAttributeTransformerTests.cpp b/Engine/Plugins/Experimental/UAF/UAFTestSuites/Source/UAFTestSuite/Private/UAF/ValueRuntime/Transformers/IntegerAttributeTransformerTests.cpp
new file mode 100644
--- /dev/null
+++ b/Engine/Plugins/Experimental/UAF/UAFTestSuites/Source/UAFTestSuite/Private/UAF/ValueRuntime/Transformers/IntegerAttributeTransformerTests.cpp
@@ -0,0 +1,302 @@
+// Copyright Epic Games, Inc. All Rights Reserved.
+
+#pragma once
+
+#include "CoreMinimal.h"
+#include "Misc/AutomationTest.h"
+
+#if WITH_DEV_AUTOMATION_TESTS
+
+#include "AnimNextTest.h"
+#include "Animation/BuiltInAttributeTypes.h"
+#include "UAF/AbstractSkeleton/AbstractSkeletonSetBinding.h"
+#include "UAF/Attributes/AttributeBindingDataCache.h"
+#include "UAF/ValueRuntime/ValueRuntimeRegistry.h"
+#include "UAF/ValueRuntime/Transformers/Accumulate.h"
+#include "UAF/ValueRuntime/Transformers/AdditiveSpace.h"
+#include "UAF/ValueRuntime/Transformers/Interpolate.h"
+#include "UAF/ValueRuntime/Transformers/Overwrite.h"
+#include "UAF/ValueRuntime/Transformers/TransformerTestUtil.h"
+#include "UObject/Package.h"
+
+namespace UE::UAF::Tests
+{
+	IMPLEMENT_SIMPLE_AUTOMATION_TEST(FTransformerTests_IntegerAttribute, "Animation.UAF.ValueRuntime.Transformers.IntegerAttribute", EAutomationTestFlags::EditorContext | EAutomationTestFlags::EngineFilter)
+
+	bool FTransformerTests_IntegerAttribute::RunTest(const FString& InParameters)
+	{
+		ON_SCOPE_EXIT{ FUtils::CleanupAfterTests(); };
+
+		TValueOrError<TNonNullPtr<UAbstractSkeletonSetBinding>, FString> SetBindingResult = BuildSetBinding({});
+
+		if (SetBindingResult.HasError())
+		{
+			AddError(SetBindingResult.GetError());
+			return false;
+		}
+		
+		TNonNullPtr<UAbstractSkeletonSetBinding> SetBinding = SetBindingResult.GetValue();
+
+		SetBinding->AddAttributeToSet(FAnimationAttributeIdentifier(TEXT("AttrA"), INDEX_NONE, NAME_None, FIntegerAnimationAttribute::StaticStruct()), "SetA");
+		SetBinding->AddAttributeToSet(FAnimationAttributeIdentifier(TEXT("AttrB"), INDEX_NONE, NAME_None, FIntegerAnimationAttribute::StaticStruct()), "SetB");
+		SetBinding->AddAttributeToSet(FAnimationAttributeIdentifier(TEXT("AttrC"), INDEX_NONE, NAME_None, FIntegerAnimationAttribute::StaticStruct()), "SetC");
+		
+		const FAttributeBindingDataPtr& BindingData = GAttributeBindingDataCache.GetOrAdd(SetBinding, nullptr);
+		FAttributeNamedSetPtr SetA = BindingData->FindNamedSet("SetA");
+		FAttributeNamedSetPtr SetB = BindingData->FindNamedSet("SetB");
+		FAttributeNamedSetPtr SetC = BindingData->FindNamedSet("SetC");
+
+		// Overwrite
+		AddInfo("Overwrite");
+		{
+			// 0.25 weight
+			{
+				FTransformerUnaryTest<FIntegerAnimationAttribute, false> IntegerTest;
+				IntegerTest.AttributeName = "AttrA";
+				IntegerTest.Input.Value = 5;
+				IntegerTest.ApplyTransformer = [](const FValueBundleHeap& InputValues, FValueBundleHeap& OutputValues)
+					{
+						const auto& TransformerMap = FValueRuntimeRegistry::Get().GetTransformerMap();
+						Transformers::FOverwrite::Apply(TransformerMap, InputValues, 0.25f, OutputValues);
+					};
+				IntegerTest.IsExpectedResult = [&](TOptional<FIntegerAnimationAttribute> Result) -> TOptional<FString>
+					{
+						if (!Result.IsSet())
+						{
+							return FString("Overwrite - 0.25 weight - expected to produce a value but received none");
+						}
+
+						constexpr int32 ExpectedValue = static_cast<int32>(0.25f * 5);
+						if (Result->Value != ExpectedValue)
+						{
+							return FString::Format(TEXT("Overwrite - 0.25 weight - expected {0}, Received {1}"), { ExpectedValue, Result->Value });
+						}
+
+						return TOptional<FString>();
+					};
+			}
+		}
+		
+		// Accumulate
+		AddInfo("Accumulate");
+		{
+			// 0.25 weight
+			{
+				FTransformerBinaryTest<Transformers::FAccumulate, FIntegerAnimationAttribute, EInplaceTransformationSupport::Allow> IntegerTest;
+				IntegerTest.AttributeName = "AttrA";
+				IntegerTest.InputA.Value = 5;
+				IntegerTest.InputB.Value = 8;
+				IntegerTest.ApplyTransformer = [](const FValueBundleHeap& InputAValues, const FValueBundleHeap& InputBValues, FValueBundleHeap& OutputValues)
+					{
+						const auto& TransformerMap = FValueRuntimeRegistry::Get().GetTransformerMap();
+						Transformers::FAccumulate::Apply(TransformerMap, InputAValues, InputBValues, 0.25f, OutputValues);
+					};
+				IntegerTest.IsExpectedResult = [&](TOptional<FIntegerAnimationAttribute> Result) -> TOptional<FString>
+					{
+						if (!Result.IsSet())
+						{
+							return FString("Accumulate - 0.25 weight - expected to produce a value but received none");
+						}
+						constexpr int32 ExpectedValue = static_cast<int32>(5 + (0.25 * 8));
+						if (Result->Value != ExpectedValue)
+						{
+							return FString::Format(TEXT("Accumulate - 0.25 weight - expected {0}, received {1}"), { ExpectedValue, Result->Value });
+						}
+						return TOptional<FString>();
+					};
+
+				IntegerTest.Run(*this, SetA);
+			}
+		}
+
+		// Interpolate
+		AddInfo("Interpolate");
+		{
+			// 0.25 weight
+			{
+				FTransformerBinaryTest<Transformers::FInterpolate, FIntegerAnimationAttribute, EInplaceTransformationSupport::Allow> IntegerTest;
+				IntegerTest.AttributeName = "AttrA";
+				IntegerTest.InputA.Value = 5;
+				IntegerTest.InputB.Value = 8;
+				IntegerTest.ApplyTransformer = [](const FValueBundleHeap& InputAValues, const FValueBundleHeap& InputBValues, FValueBundleHeap& OutputValues)
+					{
+						const auto& TransformerMap = FValueRuntimeRegistry::Get().GetTransformerMap();
+						Transformers::FInterpolate::Apply(TransformerMap, InputAValues, InputBValues, 0.25f, OutputValues);
+					};
+				IntegerTest.IsExpectedResult = [&](TOptional<FIntegerAnimationAttribute> Result) -> TOptional<FString>
+					{
+						if (!Result.IsSet())
+						{
+							return FString("Interpolate - 0.25 weight -expected to produce a value but received none");
+						}
+						constexpr int32 ExpectedValue = static_cast<int32>((0.75 * 5) + (0.25 * 8));
+						if (Result->Value != ExpectedValue)
+						{
+							return FString::Format(TEXT("Interpolate - 0.25 weight - expected {0}, received {1}"), { ExpectedValue, Result->Value });
+						}
+						return TOptional<FString>();
+					};
+
+				IntegerTest.Run(*this, SetA);
+			}
+		}
+
+		// Layer
+		AddInfo("Layer");
+		{
+			// Value in base but not in layer
+			{
+				// Weight 0.25
+				{
+					FTransformerBinaryTest<Transformers::FLayer, FIntegerAnimationAttribute, EInplaceTransformationSupport::Allow> IntegerTest;
+					IntegerTest.AttributeName = "AttrA";
+					IntegerTest.WeightB = 0.25f;
+					IntegerTest.InputA = FIntegerAnimationAttribute({ 5 });
+					IntegerTest.InputB.Reset();
+					IntegerTest.IsExpectedResult = [&](TOptional<FIntegerAnimationAttribute> Result) -> TOptional<FString>
+						{
+							if (!Result.IsSet())
+							{
+								return FString("Layer - 0.25 weight, case 1 - expected to produce a value but received none");
+							}
+							
+							constexpr int32 ExpectedValue = 5;
+							if (Result->Value != ExpectedValue)
+							{
+								return FString::Format(TEXT("Layer - 0.25 weight, case 1 - expected {0}, received {1}"), { ExpectedValue, Result->Value });
+							}
+							
+							return TOptional<FString>();
+						};
+
+					IntegerTest.Run(*this, SetA, SetB);
+				}
+			}
+
+			// Value in layer but not in base
+			{
+				// Weight 0.25
+				{
+					FTransformerBinaryTest<Transformers::FLayer, FIntegerAnimationAttribute, EInplaceTransformationSupport::Allow> IntegerTest;
+					IntegerTest.AttributeName = "AttrC";
+					IntegerTest.WeightB = 0.25f;
+					IntegerTest.InputA.Reset();
+					IntegerTest.InputB = FIntegerAnimationAttribute({ 8 });
+					IntegerTest.IsExpectedResult = [&](TOptional<FIntegerAnimationAttribute> Result) -> TOptional<FString>
+						{
+							if (Result.IsSet())
+							{
+								return FString("Layer - 0.25 weight, case 2 - expected to not produce a value but received one");
+							}
+														
+							return TOptional<FString>();
+						};
+
+					IntegerTest.Run(*this, SetB, SetC);
+				}
+			}
+			
+			// Value in both base and layer
+			{
+				// Weight 0.25
+				{
+					FTransformerBinaryTest<Transformers::FLayer, FIntegerAnimationAttribute, EInplaceTransformationSupport::Allow> IntegerTest;
+					IntegerTest.AttributeName = "AttrB";
+					IntegerTest.WeightB = 0.25f;
+					IntegerTest.InputA = FIntegerAnimationAttribute({ 5 });
+					IntegerTest.InputB = FIntegerAnimationAttribute({ 8 });
+					IntegerTest.IsExpectedResult = [&](TOptional<FIntegerAnimationAttribute> Result) -> TOptional<FString>
+						{
+							if (!Result.IsSet())
+							{
+								return FString("Layer - 0.25 weight, case 3 - expected to produce a value but received none");
+							}
+							
+							constexpr int32 ExpectedValue = static_cast<int32>((0.75f * 5) + (0.25F * 8));
+							if (Result->Value != ExpectedValue)
+							{
+								return FString::Format(TEXT("Layer - 0.25 weight, case 3 - expected {0}, received {1}"), { ExpectedValue, Result->Value });
+							}
+							
+							return TOptional<FString>();
+						};
+
+					IntegerTest.Run(*this, SetA, SetB);
+				}
+			}
+		}
+
+		// Make Additive Space
+		AddInfo("MakeAdditiveSpace");
+		{
+			// 0.25 weight
+			{
+				FTransformerBinaryTest<Transformers::FMakeAdditiveSpace, FIntegerAnimationAttribute, EInplaceTransformationSupport::Disallow> IntegerTest;
+				IntegerTest.AttributeName = "AttrA";
+				IntegerTest.InputA.Value = 5;
+				IntegerTest.InputB.Value = 8;
+				IntegerTest.ApplyTransformer = [](const FValueBundleHeap& InputAValues, const FValueBundleHeap& InputBValues, FValueBundleHeap& OutputValues)
+					{
+						const auto& TransformerMap = FValueRuntimeRegistry::Get().GetTransformerMap();
+						Transformers::FMakeAdditiveSpace::Apply(TransformerMap, InputAValues, InputBValues, OutputValues);
+					};
+				IntegerTest.IsExpectedResult = [&](TOptional<FIntegerAnimationAttribute> Result) -> TOptional<FString>
+					{
+						if (!Result.IsSet())
+						{
+							return FString("MakeAdditiveSpace - expected to produce a value but received none");
+						}
+						constexpr int32 ExpectedValue = 8 - 5;
+						if (Result->Value != ExpectedValue)
+						{
+							return FString::Format(TEXT("MakeAdditiveSpace - expected {0}, received {1}"), { ExpectedValue, Result->Value });
+						}
+						return TOptional<FString>();
+					};
+
+				IntegerTest.Run(*this, SetA);
+			}
+		}
+
+		// Apply Additive Space
+		AddInfo("ApplyAdditiveSpace");
+		{
+			// 0.25 weight
+			{
+				FTransformerBinaryTest<Transformers::FMakeAdditiveSpace, FIntegerAnimationAttribute, EInplaceTransformationSupport::Allow> IntegerTest;
+				IntegerTest.AttributeName = "AttrA";
+				IntegerTest.InputA.Value = 5;
+				IntegerTest.InputB.Value = 8;
+				IntegerTest.ApplyTransformer = [](const FValueBundleHeap& InputAValues, const FValueBundleHeap& InputBValues, FValueBundleHeap& OutputValues)
+					{
+						const_cast<FValueBundleHeap&>(InputBValues).SetValueSpace(FValueSpace(EValueSpaceType::Local, true));
+						
+						const auto& TransformerMap = FValueRuntimeRegistry::Get().GetTransformerMap();
+						Transformers::FApplyAdditiveSpace::Apply(TransformerMap, InputAValues, InputBValues, 0.25f, OutputValues);
+					};
+				IntegerTest.IsExpectedResult = [&](TOptional<FIntegerAnimationAttribute> Result) -> TOptional<FString>
+					{
+						if (!Result.IsSet())
+						{
+							return FString("ApplyAdditiveSpace - 0.25 weight - expected to produce a value but received none");
+						}
+						constexpr int32 ExpectedValue = 5 + static_cast<int32>(0.25f * 8);
+						if (Result->Value != ExpectedValue)
+						{
+							return FString::Format(TEXT("ApplyAdditiveSpace - 0.25 weight - expected {0}, received {1}"), { ExpectedValue, Result->Value });
+						}
+						return TOptional<FString>();
+					};
+
+				IntegerTest.Run(*this, SetA);
+			}
+		}
+		
+		// Sanitize
+		// Nothing to do
+		
+		return true;
+	}
+}
+
+#endif // WITH_DEV_AUTOMATION_TESTS
\ No newline at end of file
diff --git a/Engine/Plugins/Experimental/UAF/UAFTestSuites/Source/UAFTestSuite/Private/UAF/ValueRuntime/Transformers/QuaternionAttributeTransformerTests.cpp b/Engine/Plugins/Experimental/UAF/UAFTestSuites/Source/UAFTestSuite/Private/UAF/ValueRuntime/Transformers/QuaternionAttributeTransformerTests.cpp
new file mode 100644
--- /dev/null
+++ b/Engine/Plugins/Experimental/UAF/UAFTestSuites/Source/UAFTestSuite/Private/UAF/ValueRuntime/Transformers/QuaternionAttributeTransformerTests.cpp
@@ -0,0 +1,306 @@
+// Copyright Epic Games, Inc. All Rights Reserved.
+
+#pragma once
+
+#include "CoreMinimal.h"
+#include "Misc/AutomationTest.h"
+
+#if WITH_DEV_AUTOMATION_TESTS
+
+#include "AnimNextTest.h"
+#include "Animation/BuiltInAttributeTypes.h"
+#include "UAF/AbstractSkeleton/AbstractSkeletonSetBinding.h"
+#include "UAF/Attributes/AttributeBindingDataCache.h"
+#include "UAF/ValueRuntime/ValueRuntimeRegistry.h"
+#include "UAF/ValueRuntime/Transformers/Accumulate.h"
+#include "UAF/ValueRuntime/Transformers/AdditiveSpace.h"
+#include "UAF/ValueRuntime/Transformers/Interpolate.h"
+#include "UAF/ValueRuntime/Transformers/Overwrite.h"
+#include "UAF/ValueRuntime/Transformers/Sanitize.h"
+#include "UAF/ValueRuntime/Transformers/TransformerTestUtil.h"
+#include "UObject/Package.h"
+
+namespace UE::UAF::Tests
+{
+	IMPLEMENT_SIMPLE_AUTOMATION_TEST(FTransformerTests_QuaternionAttribute, "Animation.UAF.ValueRuntime.Transformers.QuaternionAttribute", EAutomationTestFlags::EditorContext | EAutomationTestFlags::EngineFilter)
+
+	bool FTransformerTests_QuaternionAttribute::RunTest(const FString& InParameters)
+	{
+		ON_SCOPE_EXIT{ FUtils::CleanupAfterTests(); };
+
+		TValueOrError<TNonNullPtr<UAbstractSkeletonSetBinding>, FString> SetBindingResult = BuildSetBinding({});
+
+		if (SetBindingResult.HasError())
+		{
+			AddError(SetBindingResult.GetError());
+			return false;
+		}
+		
+		TNonNullPtr<UAbstractSkeletonSetBinding> SetBinding = SetBindingResult.GetValue();
+
+		SetBinding->AddAttributeToSet(FAnimationAttributeIdentifier(TEXT("AttrA"), INDEX_NONE, NAME_None, FQuaternionAnimationAttribute::StaticStruct()), "SetA");
+		SetBinding->AddAttributeToSet(FAnimationAttributeIdentifier(TEXT("AttrB"), INDEX_NONE, NAME_None, FQuaternionAnimationAttribute::StaticStruct()), "SetB");
+		SetBinding->AddAttributeToSet(FAnimationAttributeIdentifier(TEXT("AttrC"), INDEX_NONE, NAME_None, FQuaternionAnimationAttribute::StaticStruct()), "SetC");
+
+		const FAttributeBindingDataPtr& BindingData = GAttributeBindingDataCache.GetOrAdd(SetBinding, nullptr);
+		FAttributeNamedSetPtr SetA = BindingData->FindNamedSet("SetA");
+		FAttributeNamedSetPtr SetB = BindingData->FindNamedSet("SetB");
+		FAttributeNamedSetPtr SetC = BindingData->FindNamedSet("SetC");
+
+		const FQuat QuatValueA = FQuat(0.0f, 1.0f, 2.0f, 0.0f);
+		const FQuat QuatValueB = FQuat(0.0f, 1.0f, 0.0f, 1.0f);
+
+		// Overwrite
+		AddInfo("Overwrite");
+		{
+			// 0.25 weight
+			{
+				FTransformerUnaryTest<FQuaternionAnimationAttribute, false> QuatTest;
+				QuatTest.AttributeName = "AttrA";
+				QuatTest.Input.Value = QuatValueA;
+				QuatTest.ApplyTransformer = [](const FValueBundleHeap& InputValues, FValueBundleHeap& OutputValues)
+					{
+						const auto& TransformerMap = FValueRuntimeRegistry::Get().GetTransformerMap();
+						Transformers::FOverwrite::Apply(TransformerMap, InputValues, 0.25f, OutputValues);
+					};
+				QuatTest.IsExpectedResult = [&](TOptional<FQuaternionAnimationAttribute> Result) -> TOptional<FString>
+					{
+						if (!Result.IsSet())
+						{
+							return FString("Overwrite - 0.25 weight - expected to produce a value but received none");
+						}
+
+						const FQuat ExpectedValue(0.0f, 0.25f, 0.5f, 0.0f);
+						if (!ExpectedValue.Equals(Result->Value))
+						{
+							return FString::Printf(TEXT("Overwrite - 0.25 weight - expected %s, Received %s"), *ExpectedValue.ToString(), *Result->Value.ToString());
+						}
+
+						return TOptional<FString>();
+					};
+
+				QuatTest.Run(*this, SetA);
+			}
+		}
+
+		// Accumulate
+		AddInfo("Accumulate");
+		{
+			// 0.25 weight
+			{
+				FTransformerBinaryTest<Transformers::FAccumulate, FQuaternionAnimationAttribute, EInplaceTransformationSupport::Allow> QuatTest;
+				QuatTest.AttributeName = "AttrA";
+				QuatTest.InputA.Value = QuatValueA;
+				QuatTest.InputB.Value = QuatValueB;
+				QuatTest.ApplyTransformer = [](const FValueBundleHeap& InputAValues, const FValueBundleHeap& InputBValues, FValueBundleHeap& OutputValues)
+					{
+						const auto& TransformerMap = FValueRuntimeRegistry::Get().GetTransformerMap();
+						Transformers::FAccumulate::Apply(TransformerMap, InputAValues, InputBValues, 0.25f, OutputValues);
+					};
+				QuatTest.IsExpectedResult = [&](TOptional<FQuaternionAnimationAttribute> Result) -> TOptional<FString>
+					{
+						if (!Result.IsSet())
+						{
+							return FString("Accumulate - 0.25 weight - expected to produce a value but received none");
+						}
+						const FQuat ExpectedValue(0.0f, 1.25f, 2.0f, 0.25f);
+						if (!ExpectedValue.Equals(Result->Value))
+						{
+							return FString::Printf(TEXT("Accumulate - 0.25 weight - expected %s, received %s"), *ExpectedValue.ToString(), *Result->Value.ToString());
+						}
+						return TOptional<FString>();
+					};
+
+				QuatTest.Run(*this, SetA);
+			}
+		}
+
+		// Interpolate
+		AddInfo("Interpolate");
+		{
+			// 0.25 weight
+			{
+				FTransformerBinaryTest<Transformers::FInterpolate, FQuaternionAnimationAttribute, EInplaceTransformationSupport::Allow> QuatTest;
+				QuatTest.AttributeName = "AttrA";
+				QuatTest.InputA.Value = QuatValueA;
+				QuatTest.InputB.Value = QuatValueB;
+				QuatTest.ApplyTransformer = [](const FValueBundleHeap& InputAValues, const FValueBundleHeap& InputBValues, FValueBundleHeap& OutputValues)
+					{
+						const auto& TransformerMap = FValueRuntimeRegistry::Get().GetTransformerMap();
+						Transformers::FInterpolate::Apply(TransformerMap, InputAValues, InputBValues, 0.25f, OutputValues);
+					};
+				QuatTest.IsExpectedResult = [&](TOptional<FQuaternionAnimationAttribute> Result) -> TOptional<FString>
+					{
+						if (!Result.IsSet())
+						{
+							return FString("Interpolate - 0.25 weight - expected to produce a value but received none");
+						}
+						const FQuat ExpectedValue = FQuat::FastLerp(QuatValueA, QuatValueB, 0.25f);
+						if (!ExpectedValue.Equals(Result->Value))
+						{
+							return FString::Printf(TEXT("Interpolate - 0.25 weight - expected %s, received %s"), *ExpectedValue.ToString(), *Result->Value.ToString());
+						}
+						return TOptional<FString>();
+					};
+
+				QuatTest.Run(*this, SetA);
+			}
+		}
+
+		// Layer
+		AddInfo("Layer");
+		{
+			// Value in layer but not in base
+			{
+				// Weight 0.25
+				{
+					FTransformerBinaryTest<Transformers::FLayer, FQuaternionAnimationAttribute, EInplaceTransformationSupport::Allow> QuatTest;
+					QuatTest.AttributeName = "AttrC";
+					QuatTest.WeightB = 0.25f;
+					QuatTest.InputA.Reset();
+					QuatTest.InputB = FQuaternionAnimationAttribute({ QuatValueB });
+					QuatTest.IsExpectedResult = [&](TOptional<FQuaternionAnimationAttribute> Result) -> TOptional<FString>
+						{
+							if (Result.IsSet())
+							{
+								return FString("Layer - 0.25 weight, case 2 - expected to not produce a value but received one");
+							}
+
+							return TOptional<FString>();
+						};
+
+					QuatTest.Run(*this, SetB, SetC);
+				}
+			}
+
+			// Value in both base and layer
+			{
+				// Weight 0.25
+				{
+					FTransformerBinaryTest<Transformers::FLayer, FQuaternionAnimationAttribute, EInplaceTransformationSupport::Allow> QuatTest;
+					QuatTest.AttributeName = "AttrB";
+					QuatTest.WeightB = 0.25f;
+					QuatTest.InputA = FQuaternionAnimationAttribute({ QuatValueA });
+					QuatTest.InputB = FQuaternionAnimationAttribute({ QuatValueB });
+					QuatTest.IsExpectedResult = [&](TOptional<FQuaternionAnimationAttribute> Result) -> TOptional<FString>
+						{
+							if (!Result.IsSet())
+							{
+								return FString("Layer - 0.25 weight, case 3 - expected to produce a value but received none");
+							}
+
+							// = (1-0.25)*FQuat(0,0,0,1) + 0.25*FQuat(0,0,1,0) = FQuat(0,0,0.25,0.75)
+							const FQuat ExpectedValue = FQuat::FastLerp(QuatValueA, QuatValueB, 0.25f);
+							if (!ExpectedValue.Equals(Result->Value))
+							{
+								return FString::Printf(TEXT("Layer - 0.25 weight, case 3 - expected %s, received %s"), *ExpectedValue.ToString(), *Result->Value.ToString());
+							}
+
+							return TOptional<FString>();
+						};
+
+					QuatTest.Run(*this, SetA, SetB);
+				}
+			}
+		}
+
+		// Make Additive Space
+		AddInfo("MakeAdditiveSpace");
+		{
+			{
+				FTransformerBinaryTest<Transformers::FMakeAdditiveSpace, FQuaternionAnimationAttribute, EInplaceTransformationSupport::Disallow> QuatTest;
+				QuatTest.AttributeName = "AttrA";
+				QuatTest.InputA.Value = QuatValueA;
+				QuatTest.InputB.Value = QuatValueB;
+				QuatTest.ApplyTransformer = [](const FValueBundleHeap& InputAValues, const FValueBundleHeap& InputBValues, FValueBundleHeap& OutputValues)
+					{
+						const auto& TransformerMap = FValueRuntimeRegistry::Get().GetTransformerMap();
+						Transformers::FMakeAdditiveSpace::Apply(TransformerMap, InputAValues, InputBValues, OutputValues);
+					};
+				QuatTest.IsExpectedResult = [&](TOptional<FQuaternionAnimationAttribute> Result) -> TOptional<FString>
+					{
+						if (!Result.IsSet())
+						{
+							return FString("MakeAdditiveSpace - expected to produce a value but received none");
+						}
+						const FQuat ExpectedValue = (QuatValueB * QuatValueA.Inverse()).GetNormalized();
+						if (!ExpectedValue.Equals(Result->Value))
+						{
+							return FString::Printf(TEXT("MakeAdditiveSpace - expected %s, received %s"), *ExpectedValue.ToString(), *Result->Value.ToString());
+						}
+						return TOptional<FString>();
+					};
+
+				QuatTest.Run(*this, SetA);
+			}
+		}
+
+		// Apply Additive Space
+		AddInfo("ApplyAdditiveSpace");
+		{
+			// 0.25 weight
+			{
+				FTransformerBinaryTest<Transformers::FMakeAdditiveSpace, FQuaternionAnimationAttribute, EInplaceTransformationSupport::Allow> QuatTest;
+				QuatTest.AttributeName = "AttrA";
+				QuatTest.InputA.Value = QuatValueA;
+				QuatTest.InputB.Value = QuatValueB;
+				QuatTest.ApplyTransformer = [](const FValueBundleHeap& InputAValues, const FValueBundleHeap& InputBValues, FValueBundleHeap& OutputValues)
+					{
+						const_cast<FValueBundleHeap&>(InputBValues).SetValueSpace(FValueSpace(EValueSpaceType::Local, true));
+
+						const auto& TransformerMap = FValueRuntimeRegistry::Get().GetTransformerMap();
+						Transformers::FApplyAdditiveSpace::Apply(TransformerMap, InputAValues, InputBValues, 0.25f, OutputValues);
+					};
+				QuatTest.IsExpectedResult = [QuatValueA, QuatValueB](TOptional<FQuaternionAnimationAttribute> Result) -> TOptional<FString>
+					{
+						if (!Result.IsSet())
+						{
+							return FString("ApplyAdditiveSpace - 0.25 weight - expected to produce a value but received none");
+						}
+						const FQuat WeightedRotation = FQuat::FastLerp(FQuat::Identity, QuatValueB, 0.25f).GetNormalized();
+						const FQuat ExpectedValue = QuatValueA + WeightedRotation;
+						if (!ExpectedValue.Equals(Result->Value))
+						{
+							return FString::Printf(TEXT("ApplyAdditiveSpace - 0.25 weight - expected %s, received %s"), *ExpectedValue.ToString(), *Result->Value.ToString());
+						}
+						return TOptional<FString>();
+					};
+
+				QuatTest.Run(*this, SetA);
+			}
+		}
+
+		// Sanitize
+		AddInfo("Sanitize");
+		{
+			FTransformerUnaryTest<FQuaternionAnimationAttribute, true> BoneTransformTest;
+			BoneTransformTest.AttributeName = "AttrA";
+			BoneTransformTest.Input.Value = FQuat(0.2f, 0.5f, 0.3f, 0.1f);
+			BoneTransformTest.ApplyTransformer = [](FValueBundleHeap& InputValues)
+				{
+					const auto& TransformerMap = FValueRuntimeRegistry::Get().GetTransformerMap();
+					Transformers::FSanitize::Apply(TransformerMap, InputValues);
+				};
+			BoneTransformTest.IsExpectedResult = [&](TOptional<FQuaternionAnimationAttribute> Result) -> TOptional<FString>
+				{
+					if (!Result.IsSet())
+					{
+						return FString("Sanitize - expected to produce a value but received none");
+					}
+					const FQuat ExpectedValue = FQuat(0.2f, 0.5f, 0.3f, 0.1f).GetNormalized();
+					if (!ExpectedValue.Equals(Result->Value))
+					{
+						return FString::Format(TEXT("Sanitize - expected {0}, Received {1}"), { ExpectedValue.ToString(), Result->Value.ToString() });
+					}
+
+					return TOptional<FString>();
+				};
+
+			BoneTransformTest.Run(*this, SetA);
+		}
+
+		return true;
+	}
+}
+
+#endif // WITH_DEV_AUTOMATION_TESTS
diff --git a/Engine/Plugins/Experimental/UAF/UAFTestSuites/Source/UAFTestSuite/Private/UAF/ValueRuntime/Transformers/TransformAttributeTransformerTests.cpp b/Engine/Plugins/Experimental/UAF/UAFTestSuites/Source/UAFTestSuite/Private/UAF/ValueRuntime/Transformers/TransformAttributeTransformerTests.cpp
new file mode 100644
--- /dev/null
+++ b/Engine/Plugins/Experimental/UAF/UAFTestSuites/Source/UAFTestSuite/Private/UAF/ValueRuntime/Transformers/TransformAttributeTransformerTests.cpp
@@ -0,0 +1,318 @@
+// Copyright Epic Games, Inc. All Rights Reserved.
+
+#pragma once
+
+#include "CoreMinimal.h"
+#include "Misc/AutomationTest.h"
+
+#if WITH_DEV_AUTOMATION_TESTS
+
+#include "AnimNextTest.h"
+#include "Animation/BuiltInAttributeTypes.h"
+#include "UAF/AbstractSkeleton/AbstractSkeletonSetBinding.h"
+#include "UAF/Attributes/AttributeBindingDataCache.h"
+#include "UAF/ValueRuntime/ValueRuntimeRegistry.h"
+#include "UAF/ValueRuntime/Transformers/Accumulate.h"
+#include "UAF/ValueRuntime/Transformers/AdditiveSpace.h"
+#include "UAF/ValueRuntime/Transformers/Interpolate.h"
+#include "UAF/ValueRuntime/Transformers/Overwrite.h"
+#include "UAF/ValueRuntime/Transformers/Sanitize.h"
+#include "UAF/ValueRuntime/Transformers/TransformerTestUtil.h"
+#include "UObject/Package.h"
+
+namespace UE::UAF::Tests
+{
+	IMPLEMENT_SIMPLE_AUTOMATION_TEST(FTransformerTests_TransformAttribute, "Animation.UAF.ValueRuntime.Transformers.TransformAttribute", EAutomationTestFlags::EditorContext | EAutomationTestFlags::EngineFilter)
+
+	bool FTransformerTests_TransformAttribute::RunTest(const FString& InParameters)
+	{
+		ON_SCOPE_EXIT{ FUtils::CleanupAfterTests(); };
+
+		TValueOrError<TNonNullPtr<UAbstractSkeletonSetBinding>, FString> SetBindingResult = BuildSetBinding({});
+
+		if (SetBindingResult.HasError())
+		{
+			AddError(SetBindingResult.GetError());
+			return false;
+		}
+		
+		TNonNullPtr<UAbstractSkeletonSetBinding> SetBinding = SetBindingResult.GetValue();
+
+		SetBinding->AddAttributeToSet(FAnimationAttributeIdentifier(TEXT("AttrA"), INDEX_NONE, NAME_None, FTransformAnimationAttribute::StaticStruct()), "SetA");
+		SetBinding->AddAttributeToSet(FAnimationAttributeIdentifier(TEXT("AttrB"), INDEX_NONE, NAME_None, FTransformAnimationAttribute::StaticStruct()), "SetB");
+		SetBinding->AddAttributeToSet(FAnimationAttributeIdentifier(TEXT("AttrC"), INDEX_NONE, NAME_None, FTransformAnimationAttribute::StaticStruct()), "SetC");
+
+		const FAttributeBindingDataPtr& BindingData = GAttributeBindingDataCache.GetOrAdd(SetBinding, nullptr);
+		FAttributeNamedSetPtr SetA = BindingData->FindNamedSet("SetA");
+		FAttributeNamedSetPtr SetB = BindingData->FindNamedSet("SetB");
+		FAttributeNamedSetPtr SetC = BindingData->FindNamedSet("SetC");
+
+		const FTransform TransformValueA(FQuat::Identity, FVector(100.0f, 200.0f, 300.0f), FVector::OneVector);
+		const FTransform TransformValueB(FQuat(0.0f, 1.0f, 0.0f, 0.0f), FVector(400.0f, 500.0f, 600.0f), FVector(2.0f, 2.0f, 2.0f));
+
+		// Overwrite
+		AddInfo("Overwrite");
+		{
+			// 0.25 weight
+			{
+				FTransformerUnaryTest<FTransformAnimationAttribute, false> TransformTest;
+				TransformTest.AttributeName = "AttrA";
+				TransformTest.Input.Value = TransformValueA;
+				TransformTest.ApplyTransformer = [](const FValueBundleHeap& InputValues, FValueBundleHeap& OutputValues)
+					{
+						const auto& TransformerMap = FValueRuntimeRegistry::Get().GetTransformerMap();
+						Transformers::FOverwrite::Apply(TransformerMap, InputValues, 0.25f, OutputValues);
+					};
+				TransformTest.IsExpectedResult = [&](TOptional<FTransformAnimationAttribute> Result) -> TOptional<FString>
+					{
+						if (!Result.IsSet())
+						{
+							return FString("Overwrite - 0.25 weight - expected to produce a value but received none");
+						}
+						
+						const FTransform ExpectedValue = TransformValueA * ScalarRegister(0.25f);
+						if (!ExpectedValue.Equals(Result->Value))
+						{
+							return FString::Printf(TEXT("Overwrite - 0.25 weight - expected %s, Received %s"), *ExpectedValue.ToHumanReadableString(), *Result->Value.ToHumanReadableString());
+						}
+
+						return TOptional<FString>();
+					};
+
+				TransformTest.Run(*this, SetA);
+			}
+		}
+
+		// Accumulate
+		AddInfo("Accumulate");
+		{
+			// 0.25 weight
+			{
+				FTransformerBinaryTest<Transformers::FAccumulate, FTransformAnimationAttribute, EInplaceTransformationSupport::Allow> TransformTest;
+				TransformTest.AttributeName = "AttrA";
+				TransformTest.InputA.Value = TransformValueA;
+				TransformTest.InputB.Value = TransformValueB;
+				TransformTest.ApplyTransformer = [](const FValueBundleHeap& InputAValues, const FValueBundleHeap& InputBValues, FValueBundleHeap& OutputValues)
+					{
+						const auto& TransformerMap = FValueRuntimeRegistry::Get().GetTransformerMap();
+						Transformers::FAccumulate::Apply(TransformerMap, InputAValues, InputBValues, 0.25f, OutputValues);
+					};
+				TransformTest.IsExpectedResult = [&](TOptional<FTransformAnimationAttribute> Result) -> TOptional<FString>
+					{
+						if (!Result.IsSet())
+						{
+							return FString("Accumulate - 0.25 weight - expected to produce a value but received none");
+						}
+						
+						const FTransform ExpectedValue = TransformValueA + (TransformValueB * ScalarRegister(0.25f));
+						if (!ExpectedValue.Equals(Result->Value))
+						{
+							return FString::Printf(TEXT("Accumulate - 0.25 weight - expected %s, received %s"), *ExpectedValue.ToHumanReadableString(), *Result->Value.ToHumanReadableString());
+						}
+						
+						return TOptional<FString>();
+					};
+
+				TransformTest.Run(*this, SetA);
+			}
+		}
+
+		// Interpolate
+		AddInfo("Interpolate");
+		{
+			// 0.25 weight
+			{
+				FTransformerBinaryTest<Transformers::FInterpolate, FTransformAnimationAttribute, EInplaceTransformationSupport::Allow> TransformTest;
+				TransformTest.AttributeName = "AttrA";
+				TransformTest.InputA.Value = TransformValueA;
+				TransformTest.InputB.Value = TransformValueB;
+				TransformTest.ApplyTransformer = [](const FValueBundleHeap& InputAValues, const FValueBundleHeap& InputBValues, FValueBundleHeap& OutputValues)
+					{
+						const auto& TransformerMap = FValueRuntimeRegistry::Get().GetTransformerMap();
+						Transformers::FInterpolate::Apply(TransformerMap, InputAValues, InputBValues, 0.25f, OutputValues);
+					};
+				TransformTest.IsExpectedResult = [&](TOptional<FTransformAnimationAttribute> Result) -> TOptional<FString>
+					{
+						if (!Result.IsSet())
+						{
+							return FString("Interpolate - 0.25 weight - expected to produce a value but received none");
+						}
+						
+						FTransform ExpectedValue = TransformValueA;
+						ExpectedValue.BlendWith(TransformValueB, 0.25f);
+						if (!ExpectedValue.Equals(Result->Value))
+						{
+							return FString::Printf(TEXT("Interpolate - 0.25 weight - expected %s, received %s"), *ExpectedValue.ToHumanReadableString(), *Result->Value.ToHumanReadableString());
+						}
+						
+						return TOptional<FString>();
+					};
+
+				TransformTest.Run(*this, SetA);
+			}
+		}
+
+		// Layer
+		AddInfo("Layer");
+		{
+			// Value in layer but not in base
+			{
+				// Weight 0.25
+				{
+					FTransformerBinaryTest<Transformers::FLayer, FTransformAnimationAttribute, EInplaceTransformationSupport::Allow> TransformTest;
+					TransformTest.AttributeName = "AttrC";
+					TransformTest.WeightB = 0.25f;
+					TransformTest.InputA.Reset();
+					TransformTest.InputB = FTransformAnimationAttribute({ TransformValueB });
+					TransformTest.IsExpectedResult = [&](TOptional<FTransformAnimationAttribute> Result) -> TOptional<FString>
+						{
+							if (Result.IsSet())
+							{
+								return FString("Layer - 0.25 weight, case 2 - expected to not produce a value but received one");
+							}
+
+							return TOptional<FString>();
+						};
+
+					TransformTest.Run(*this, SetB, SetC);
+				}
+			}
+
+			// Value in both base and layer
+			{
+				// Weight 0.25
+				{
+					FTransformerBinaryTest<Transformers::FLayer, FTransformAnimationAttribute, EInplaceTransformationSupport::Allow> TransformTest;
+					TransformTest.AttributeName = "AttrB";
+					TransformTest.WeightB = 0.25f;
+					TransformTest.InputA = FTransformAnimationAttribute({ TransformValueA });
+					TransformTest.InputB = FTransformAnimationAttribute({ TransformValueB });
+					TransformTest.IsExpectedResult = [&](TOptional<FTransformAnimationAttribute> Result) -> TOptional<FString>
+						{
+							if (!Result.IsSet())
+							{
+								return FString("Layer - 0.25 weight, case 3 - expected to produce a value but received none");
+							}
+							
+							FTransform ExpectedValue = TransformValueA;
+							ExpectedValue.BlendWith(TransformValueB, 0.25f);
+							if (!ExpectedValue.Equals(Result->Value))
+							{
+								return FString::Printf(TEXT("Layer - 0.25 weight, case 3 - expected %s, received %s"), *ExpectedValue.ToHumanReadableString(), *Result->Value.ToHumanReadableString());
+							}
+
+							return TOptional<FString>();
+						};
+
+					TransformTest.Run(*this, SetA, SetB);
+				}
+			}
+		}
+
+		// Make Additive Space
+		AddInfo("MakeAdditiveSpace");
+		{
+			{
+				FTransformerBinaryTest<Transformers::FMakeAdditiveSpace, FTransformAnimationAttribute, EInplaceTransformationSupport::Disallow> TransformTest;
+				TransformTest.AttributeName = "AttrA";
+				TransformTest.InputA.Value = TransformValueA;
+				TransformTest.InputB.Value = TransformValueB;
+				TransformTest.ApplyTransformer = [](const FValueBundleHeap& InputAValues, const FValueBundleHeap& InputBValues, FValueBundleHeap& OutputValues)
+					{
+						const auto& TransformerMap = FValueRuntimeRegistry::Get().GetTransformerMap();
+						Transformers::FMakeAdditiveSpace::Apply(TransformerMap, InputAValues, InputBValues, OutputValues);
+					};
+				TransformTest.IsExpectedResult = [&](TOptional<FTransformAnimationAttribute> Result) -> TOptional<FString>
+					{
+						if (!Result.IsSet())
+						{
+							return FString("MakeAdditiveSpace - expected to produce a value but received none");
+						}
+						
+						FTransform ExpectedValue = TransformValueB;
+						FAnimationRuntime::ConvertTransformToAdditive(ExpectedValue, TransformValueA);
+						if (!ExpectedValue.Equals(Result->Value))
+						{
+							return FString::Printf(TEXT("MakeAdditiveSpace - expected %s, received %s"), *ExpectedValue.ToHumanReadableString(), *Result->Value.ToHumanReadableString());
+						}
+						
+						return TOptional<FString>();
+					};
+
+				TransformTest.Run(*this, SetA);
+			}
+		}
+
+		// Apply Additive Space
+		AddInfo("ApplyAdditiveSpace");
+		{
+			// 0.25 weight
+			{
+				FTransformerBinaryTest<Transformers::FMakeAdditiveSpace, FTransformAnimationAttribute, EInplaceTransformationSupport::Allow> TransformTest;
+				TransformTest.AttributeName = "AttrA";
+				TransformTest.InputA.Value = TransformValueA;
+				TransformTest.InputB.Value = TransformValueB;
+				TransformTest.ApplyTransformer = [](const FValueBundleHeap& InputAValues, const FValueBundleHeap& InputBValues, FValueBundleHeap& OutputValues)
+					{
+						const_cast<FValueBundleHeap&>(InputBValues).SetValueSpace(FValueSpace(EValueSpaceType::Local, true));
+
+						const auto& TransformerMap = FValueRuntimeRegistry::Get().GetTransformerMap();
+						Transformers::FApplyAdditiveSpace::Apply(TransformerMap, InputAValues, InputBValues, 0.25f, OutputValues);
+					};
+				TransformTest.IsExpectedResult = [&](TOptional<FTransformAnimationAttribute> Result) -> TOptional<FString>
+					{
+						if (!Result.IsSet())
+						{
+							return FString("ApplyAdditiveSpace - 0.25 weight - expected to produce a value but received none");
+						}
+						
+						FTransform ExpectedValue = TransformValueA;
+						FTransform::BlendFromIdentityAndAccumulate(ExpectedValue, TransformValueB, ScalarRegister(0.25f));
+						if (!ExpectedValue.Equals(Result->Value))
+						{
+							return FString::Printf(TEXT("ApplyAdditiveSpace - 0.25 weight - expected %s, received %s"), *ExpectedValue.ToHumanReadableString(), *Result->Value.ToHumanReadableString());
+						}
+						
+						return TOptional<FString>();
+					};
+
+				TransformTest.Run(*this, SetA);
+			}
+		}
+
+		// Sanitize
+		AddInfo("Sanitize");
+		{
+			FTransformerUnaryTest<FTransformAnimationAttribute, true> TransformTest;
+			TransformTest.AttributeName = "AttrA";
+			TransformTest.Input.Value = FTransform(FQuat(0.2f, 0.5f, 0.3f, 0.1f), FVector(100.0f, 50.0f, 25.0f), FVector(1.0f));
+			TransformTest.ApplyTransformer = [](FValueBundleHeap& InputValues)
+				{
+					const auto& TransformerMap = FValueRuntimeRegistry::Get().GetTransformerMap();
+					Transformers::FSanitize::Apply(TransformerMap, InputValues);
+				};
+			TransformTest.IsExpectedResult = [&](TOptional<FTransformAnimationAttribute> Result) -> TOptional<FString>
+				{
+					if (!Result.IsSet())
+					{
+						return FString("Sanitize - expected to produce a value but received none");
+					}
+
+					FTransform ExpectedValue = FTransform(FQuat(0.2f, 0.5f, 0.3f, 0.1f), FVector(100.0f, 50.0f, 25.0f), FVector(1.0f));
+					ExpectedValue.NormalizeRotation();
+					if (!ExpectedValue.Equals(Result->Value))
+					{
+						return FString::Format(TEXT("Sanitize - expected {0}, Received {1}"), { ExpectedValue.ToString(), Result->Value.ToString() });
+					}
+
+					return TOptional<FString>();
+				};
+
+			TransformTest.Run(*this, SetA);
+		}
+
+		return true;
+	}
+}
+
+#endif // WITH_DEV_AUTOMATION_TESTS
diff --git a/Engine/Plugins/Experimental/UAF/UAFTestSuites/Source/UAFTestSuite/Private/UAF/ValueRuntime/Transformers/TransformerTestUtil.h b/Engine/Plugins/Experimental/UAF/UAFTestSuites/Source/UAFTestSuite/Private/UAF/ValueRuntime/Transformers/TransformerTestUtil.h
new file mode 100644
--- /dev/null
+++ b/Engine/Plugins/Experimental/UAF/UAFTestSuites/Source/UAFTestSuite/Private/UAF/ValueRuntime/Transformers/TransformerTestUtil.h
@@ -0,0 +1,422 @@
+// Copyright Epic Games, Inc. All Rights Reserved.
+
+#pragma once
+
+#include "Factories/Factory.h"
+#include "Misc/AutomationTest.h"
+#include "UAF/AbstractSkeleton/AbstractSkeletonSetBinding.h"
+#include "UAF/AbstractSkeleton/AbstractSkeletonSetCollection.h"
+#include "UAF/AbstractSkeleton/SetBindingFactory.h"
+#include "UAF/AbstractSkeleton/SetCollectionFactory.h"
+#include "UAF/ValueRuntime/ValueRuntimeRegistry.h"
+#include "UAF/ValueRuntime/Transformers/Layer.h"
+#include "UObject/Package.h"
+
+namespace UE::UAF::Tests
+{
+	// Creates a new set binding for testing
+	// Optionally can provide a function to operate on the skeleton to apply any modifications if needed
+	// SetBinding will be created with the following set hierarchy:
+	// (Everything Set)
+	// \_ SetA
+	//    \_ SetB
+	//    \_ SetC
+	inline TValueOrError<TNonNullPtr<UAbstractSkeletonSetBinding>, FString> BuildSetBinding(TFunction<void(TNonNullPtr<USkeleton>)> SkeletonModifier)
+	{
+		UFactory* BindingFactory = NewObject<UFactory>(GetTransientPackage(), UAbstractSkeletonSetBindingFactory::StaticClass());
+		if (!BindingFactory)
+		{
+			return MakeError("Failed to create factory");
+		}
+
+		UAbstractSkeletonSetBinding* SetBinding = CastChecked<UAbstractSkeletonSetBinding>(BindingFactory->FactoryCreateNew(UAbstractSkeletonSetBinding::StaticClass(), GetTransientPackage(), TEXT("TestBindingAsset"), RF_Transient, nullptr, nullptr, NAME_None));
+		if (!SetBinding)
+		{
+			return MakeError("Failed to create binding asset");
+		}
+
+		UFactory* SetCollectionFactory = NewObject<UFactory>(GetTransientPackage(), UAbstractSkeletonSetCollectionFactory::StaticClass());
+		if (!SetCollectionFactory)
+		{
+			return MakeError("Failed to create factory");
+		}
+
+		UAbstractSkeletonSetCollection* SetCollection = CastChecked<UAbstractSkeletonSetCollection>(SetCollectionFactory->FactoryCreateNew(UAbstractSkeletonSetCollection::StaticClass(), GetTransientPackage(), TEXT("TestSetCollectionAsset"), RF_Transient, nullptr, nullptr, NAME_None));
+		if (!SetCollection)
+		{
+			return MakeError("Failed to create set collection asset");
+		}
+
+		SetCollection->AddSet("SetA");
+		SetCollection->AddSet("SetB", "SetA");
+		SetCollection->AddSet("SetC", "SetA");
+		
+		bool bSuccess = SetBinding->SetSetCollection(SetCollection);
+		if (!bSuccess)
+		{
+			return MakeError("Failed to set set collection on our binding");
+		}
+
+		USkeleton* Skeleton = NewObject<USkeleton>(GetTransientPackage(), USkeleton::StaticClass(), TEXT("TestSkeletonAsset"), RF_Transient);
+		if (!Skeleton)
+		{
+			return MakeError("Failed to create skeleton asset");
+		}
+
+		if (SkeletonModifier.IsSet())
+		{
+			SkeletonModifier(Skeleton);
+		}
+		
+		bSuccess = SetBinding->SetSkeleton(Skeleton);
+		if (!bSuccess)
+		{
+			return MakeError("Failed to set skeleton on our binding");
+		}
+
+		return MakeValue(SetBinding);
+	}
+	
+	// Util struct to perform a transformer on an attribute of a particular value and to validate its result.
+	// FTransformUnaryTests is for transformers taking a single input.
+	// The same operation is performed on a bound value and unbound value.
+	// When calling Run we assume there is a valid attribute defined in the bound maps for the test to use.
+	// The SingleBuffer template argument determines whether the transformer only takes a single value bundle
+	// as input (which only supports in-place operations) or whether the transformer can output to a second
+	// bundle (which supports in-place and out-of-place operations)
+	template <typename AttributeType, bool SingleBuffer>
+	struct FTransformerUnaryTest
+	{
+		FName AttributeName;
+		AttributeType Input;
+
+		using ApplyTransformerFunctionType = std::conditional_t<SingleBuffer,
+		TFunction<void(FValueBundleHeap& InputOutputValues)>,
+		TFunction<void(const FValueBundleHeap& InputValues, FValueBundleHeap& OutputValues)>>;
+		
+		ApplyTransformerFunctionType ApplyTransformer;
+		TFunction<TOptional<FString>(TOptional<AttributeType> Result)> IsExpectedResult;
+
+		void Run(FAutomationTestBase& Test, const FAttributeNamedSetPtr& NamedSet)
+		{
+			FValueBundleHeap InputValueBundle(NamedSet);
+
+			// Setup bound attributes
+			const FAttributeMappingKey MappingKey = FAttributeMappingKey::MakeFromTo<AttributeType>();
+			const FAttributeTypedSetPtr& TypedSet = NamedSet->FindTypedSet<AttributeType>();
+			const FAttributeSetIndex AttributeSetIndex = TypedSet->FindIndex(AttributeName);
+			TBoundValueMap<AttributeType>* InputABoundValues = InputValueBundle.GetBoundValueMaps().FindOrAdd<AttributeType>(MappingKey);
+			InputABoundValues->SetValue(AttributeSetIndex, Input);
+
+			// Setup unbound attributes
+			InputValueBundle.GetUnboundValueMaps().FindOrAdd<AttributeType>()->Add(AttributeName, Input);
+
+			if constexpr (SingleBuffer)
+			{
+				ApplyTransformer(InputValueBundle);
+				ValidateResult(Test, NamedSet, "OutOfPlace", InputValueBundle);
+			}
+			else
+			{
+				// Transform input into output
+				{
+					FValueBundleHeap OutputValueBundle(NamedSet);
+					ApplyTransformer(InputValueBundle, OutputValueBundle);
+					ValidateResult(Test, NamedSet, "OutOfPlace", OutputValueBundle);
+				}
+
+				// Transform input inplace
+				{
+					FValueBundleHeap InputCopyValueBundle(NamedSet);
+					InputCopyValueBundle = InputValueBundle;
+					ApplyTransformer(InputCopyValueBundle, InputCopyValueBundle);
+					ValidateResult(Test, NamedSet, "InplaceA", InputCopyValueBundle);
+				}
+			}
+		};
+
+		void ValidateResult(FAutomationTestBase& Test, const FAttributeNamedSetPtr& NamedSet, FString ErrorPrefix, const FValueBundleHeap& Result)
+		{
+			// Validate bound attribute
+			{
+				TOptional<AttributeType> ResultAttribute;
+				
+				const FAttributeMappingKey MappingKey = FAttributeMappingKey::MakeFromTo<AttributeType>();
+				if (const FAttributeTypedSetPtr& TypedSet = NamedSet->FindTypedSet<AttributeType>())
+				{
+					if (const FAttributeSetIndex AttributeSetIndex = TypedSet->FindIndex(AttributeName))
+					{
+						if (const TBoundValueMap<AttributeType>* OutputBoundValues = Result.GetBoundValueMaps().Find<AttributeType>(MappingKey))
+						{
+							ResultAttribute = OutputBoundValues->GetValue(AttributeSetIndex);
+						}
+					}
+				}
+
+				TOptional<FString> MaybeError = IsExpectedResult(ResultAttribute);
+
+				if (MaybeError.IsSet())
+				{
+					Test.AddError(FString::Format(TEXT("{0} Bound: {1}"), { ErrorPrefix, *MaybeError }), 1);
+				}
+				else
+				{
+					Test.AddInfo(FString::Format(TEXT("{0} Bound: Pass!"), { ErrorPrefix }));
+				}
+			}
+			
+			{
+				TOptional<AttributeType> ResultAttribute;
+
+				if (const TUnboundValueMap<AttributeType>* UnboundMap = Result.GetUnboundValueMaps().Find<AttributeType>())
+				{
+					if (const AttributeType* ResultValue = UnboundMap->Find(AttributeName))
+					{
+						ResultAttribute = *ResultValue;
+					}
+				}
+				
+				TOptional<FString> MaybeError = IsExpectedResult(ResultAttribute);
+				if (MaybeError.IsSet())
+				{
+					Test.AddError(FString::Format(TEXT("{0} Unbound: {1}"), { ErrorPrefix, *MaybeError }), 1);
+				}
+				else
+				{
+					Test.AddInfo(FString::Format(TEXT("{0} Unbound: Pass!"), { ErrorPrefix }));
+				}
+			}
+		}
+	};
+
+	enum class EInplaceTransformationSupport
+	{
+		Allow,
+		Disallow,
+	};
+
+	// Util struct to perform a transformer on an attribute of a particular value and to validate its result.
+	// FTransformBinaryTest is for transformers taking two inputs.
+	// The same operation is performed on a bound value and unbound value.
+	// When calling Run we assume there is a valid attribute defined in the bound maps for the test to use.
+	// The InplaceSupport template argument determines whether the transformer allows for inplace operations.
+	template <typename TransformerType, typename AttributeType, EInplaceTransformationSupport InplaceSupport>
+	struct FTransformerBinaryTest
+	{
+		FName AttributeName;
+		AttributeType InputA;
+		AttributeType InputB;
+		
+		using ApplyTransformerFunctionType = TFunction<void(const FValueBundleHeap& InputAValues, const FValueBundleHeap& InputBValues, FValueBundleHeap& OutputValues)>;
+		ApplyTransformerFunctionType ApplyTransformer;
+		
+		TFunction<TOptional<FString>(TOptional<AttributeType> Result)> IsExpectedResult;
+
+		void Run(FAutomationTestBase& Test, FAttributeNamedSetPtr NamedSet)
+		{
+			FValueBundleHeap InputAValueBundle(NamedSet);
+			InputAValueBundle.SetValueSpace(FValueSpace(EValueSpaceType::Local));
+			FValueBundleHeap InputBValueBundle(NamedSet);
+			InputBValueBundle.SetValueSpace(FValueSpace(EValueSpaceType::Local));
+
+			// Setup bound attributes
+			const FAttributeMappingKey MappingKey = FAttributeMappingKey::MakeFromTo<AttributeType>();
+			const FAttributeTypedSetPtr& TypedSet = NamedSet->FindTypedSet<AttributeType>();
+			const FAttributeSetIndex AttributeSetIndex = TypedSet->FindIndex(AttributeName);
+			TBoundValueMap<AttributeType>* InputABoundValues = InputAValueBundle.GetBoundValueMaps().FindOrAdd<AttributeType>(MappingKey);
+			TBoundValueMap<AttributeType>* InputBBoundValues = InputBValueBundle.GetBoundValueMaps().FindOrAdd<AttributeType>(MappingKey);
+			InputABoundValues->SetValue(AttributeSetIndex, InputA);
+			InputBBoundValues->SetValue(AttributeSetIndex, InputB);
+
+			// Setup unbound attributes
+			InputAValueBundle.GetUnboundValueMaps().FindOrAdd<AttributeType>()->Add(AttributeName, InputA);
+			InputBValueBundle.GetUnboundValueMaps().FindOrAdd<AttributeType>()->Add(AttributeName, InputB);
+
+			// Transform inputs A and B into Output
+			FValueBundleHeap OutputValueBundle(NamedSet);
+			ApplyTransformer(InputAValueBundle, InputBValueBundle, OutputValueBundle);
+			ValidateResult(Test, NamedSet, "OutOfPlace", OutputValueBundle);
+
+			if constexpr (InplaceSupport == EInplaceTransformationSupport::Allow)
+			{
+				// Transform inputs A and B inplace into A
+				FValueBundleHeap InputACopyValueBundle(NamedSet);
+				InputACopyValueBundle = InputAValueBundle;
+				ApplyTransformer(InputACopyValueBundle, InputBValueBundle, InputACopyValueBundle);
+				ValidateResult(Test, NamedSet, "InplaceA", InputACopyValueBundle);
+
+				// Transform inputs A and B inplace into B
+				FValueBundleHeap InputBCopyValueBundle(NamedSet);
+				InputBCopyValueBundle = InputBValueBundle;
+				ApplyTransformer(InputAValueBundle, InputBCopyValueBundle, InputBCopyValueBundle);
+				ValidateResult(Test, NamedSet, "InplaceB", InputBCopyValueBundle);
+			}
+		};
+
+		void ValidateResult(FAutomationTestBase& Test, const FAttributeNamedSetPtr& NamedSet, FString ErrorPrefix, const FValueBundleHeap& Result)
+		{
+			// Validate bound attribute
+			{
+				TOptional<AttributeType> ResultAttribute;
+				
+				const FAttributeMappingKey MappingKey = FAttributeMappingKey::MakeFromTo<AttributeType>();
+				if (const FAttributeTypedSetPtr& TypedSet = NamedSet->FindTypedSet<AttributeType>())
+				{
+					if (const FAttributeSetIndex AttributeSetIndex = TypedSet->FindIndex(AttributeName))
+					{
+						if (const TBoundValueMap<AttributeType>* OutputBoundValues = Result.GetBoundValueMaps().Find<AttributeType>(MappingKey))
+						{
+							ResultAttribute = OutputBoundValues->GetValue(AttributeSetIndex);
+						}
+					}
+				}
+
+				TOptional<FString> MaybeError = IsExpectedResult(ResultAttribute);
+
+				if (MaybeError.IsSet())
+				{
+					Test.AddError(FString::Format(TEXT("{0} Bound: {1}"), { ErrorPrefix, *MaybeError }), 1);
+				}
+				else
+				{
+					Test.AddInfo(FString::Format(TEXT("{0} Bound: Pass!"), { ErrorPrefix }));
+				}
+			}
+			
+			{
+				TOptional<AttributeType> ResultAttribute;
+
+				if (const TUnboundValueMap<AttributeType>* UnboundMap = Result.GetUnboundValueMaps().Find<AttributeType>())
+				{
+					if (const AttributeType* ResultValue = UnboundMap->Find(AttributeName))
+					{
+						ResultAttribute = *ResultValue;
+					}
+				}
+				
+				TOptional<FString> MaybeError = IsExpectedResult(ResultAttribute);
+				if (MaybeError.IsSet())
+				{
+					Test.AddError(FString::Format(TEXT("{0} Unbound: {1}"), { ErrorPrefix, *MaybeError }), 1);
+				}
+				else
+				{
+					Test.AddInfo(FString::Format(TEXT("{0} Unbound: Pass!"), { ErrorPrefix }));
+				}
+			}
+		}
+	};
+
+	// Specialization for FLayer
+	template <typename AttributeType>
+	struct FTransformerBinaryTest<Transformers::FLayer, AttributeType, EInplaceTransformationSupport::Allow>
+	{
+		FName AttributeName;
+		float WeightB;
+		TOptional<AttributeType> InputA;
+		TOptional<AttributeType> InputB;
+		TFunction<TOptional<FString>(TOptional<AttributeType> Result)> IsExpectedResult;
+
+		void Run(FAutomationTestBase& Test, const FAttributeNamedSetPtr& NamedSetBase, const FAttributeNamedSetPtr& NamedSetLayer)
+		{
+			FValueBundleHeap InputAValueBundle(NamedSetBase);
+			InputAValueBundle.SetValueSpace(FValueSpace(EValueSpaceType::Local));
+			FValueBundleHeap InputBValueBundle(NamedSetLayer);
+			InputBValueBundle.SetValueSpace(FValueSpace(EValueSpaceType::Local));
+			
+			// Setup bound attributes
+			const FAttributeMappingKey MappingKey = FAttributeMappingKey::MakeFromTo<AttributeType>();
+			if (InputA.IsSet())
+			{
+				const FAttributeTypedSetPtr& TypedSet = NamedSetBase->FindTypedSet<AttributeType>();
+				const FAttributeSetIndex AttributeSetIndex = TypedSet->FindIndex(AttributeName);
+				TBoundValueMap<AttributeType>* InputABoundValues = InputAValueBundle.GetBoundValueMaps().FindOrAdd<AttributeType>(MappingKey);
+				InputABoundValues->SetValue(AttributeSetIndex, *InputA);
+			}
+			if (InputB.IsSet())
+			{
+				const FAttributeTypedSetPtr& TypedSet = NamedSetLayer->FindTypedSet<AttributeType>();
+				const FAttributeSetIndex AttributeSetIndex = TypedSet->FindIndex(AttributeName);
+				TBoundValueMap<AttributeType>* InputBBoundValues = InputBValueBundle.GetBoundValueMaps().FindOrAdd<AttributeType>(MappingKey);
+				InputBBoundValues->SetValue(AttributeSetIndex, *InputB);
+			}
+
+			// Setup unbound attributes
+			if (InputA.IsSet())
+			{
+				InputAValueBundle.GetUnboundValueMaps().FindOrAdd<AttributeType>()->Add(AttributeName, *InputA);
+			}
+			if (InputB.IsSet())
+			{
+				InputBValueBundle.GetUnboundValueMaps().FindOrAdd<AttributeType>()->Add(AttributeName, *InputB);
+			}
+
+			const auto& TransformerMap = FValueRuntimeRegistry::Get().GetTransformerMap();
+
+			// Transform inputs A and B into Output
+			FValueBundleHeap OutputValueBundle(NamedSetBase);
+			Transformers::FLayer::Apply(TransformerMap, InputAValueBundle, InputBValueBundle, WeightB, OutputValueBundle);
+			ValidateResult(Test, NamedSetBase, "OutOfPlace", OutputValueBundle);
+
+			// Transform inputs A and B inplace into A
+			FValueBundleHeap InputACopyValueBundle(NamedSetBase);
+			InputACopyValueBundle = InputAValueBundle;
+			Transformers::FLayer::Apply(TransformerMap, InputACopyValueBundle, InputBValueBundle, WeightB, InputACopyValueBundle);
+			ValidateResult(Test, NamedSetBase, "InplaceA", InputACopyValueBundle);
+		};
+
+		void ValidateResult(FAutomationTestBase& Test, const FAttributeNamedSetPtr& NamedSet, FString ErrorPrefix, const FValueBundleHeap& Result)
+		{
+			// Validate bound attribute
+			{
+				TOptional<AttributeType> ResultAttribute;
+				
+				const FAttributeMappingKey MappingKey = FAttributeMappingKey::MakeFromTo<AttributeType>();
+				if (const FAttributeTypedSetPtr& TypedSet = NamedSet->FindTypedSet<AttributeType>())
+				{
+					if (const FAttributeSetIndex AttributeSetIndex = TypedSet->FindIndex(AttributeName))
+					{
+						if (const TBoundValueMap<AttributeType>* OutputBoundValues = Result.GetBoundValueMaps().Find<AttributeType>(MappingKey))
+						{
+							ResultAttribute = OutputBoundValues->GetValue(AttributeSetIndex);
+						}
+					}
+				}
+
+				TOptional<FString> MaybeError = IsExpectedResult(ResultAttribute);
+
+				if (MaybeError.IsSet())
+				{
+					Test.AddError(FString::Format(TEXT("{0} Bound: {1}"), { ErrorPrefix, *MaybeError }), 1);
+				}
+				else
+				{
+					Test.AddInfo(FString::Format(TEXT("{0} Bound: Pass!"), { ErrorPrefix }));
+				}
+			}
+			
+			// Validate unbound attribute
+			{
+				TOptional<AttributeType> ResultAttribute;
+
+				if (const TUnboundValueMap<AttributeType>* UnboundMap = Result.GetUnboundValueMaps().Find<AttributeType>())
+				{
+					if (const AttributeType* ResultValue = UnboundMap->Find(AttributeName))
+					{
+						ResultAttribute = *ResultValue;
+					}
+				}
+				
+				TOptional<FString> MaybeError = IsExpectedResult(ResultAttribute);
+				if (MaybeError.IsSet())
+				{
+					Test.AddError(FString::Format(TEXT("{0} Unbound: {1}"), { ErrorPrefix, *MaybeError }), 1);
+				}
+				else
+				{
+					Test.AddInfo(FString::Format(TEXT("{0} Unbound: Pass!"), { ErrorPrefix }));
+				}
+			}
+		}
+	};
+}
\ No newline at end of file
diff --git a/Engine/Plugins/Experimental/UAF/UAFTestSuites/Source/UAFTestSuite/Private/UAF/ValueRuntime/Transformers/VectorAttributeTransformerTests.cpp b/Engine/Plugins/Experimental/UAF/UAFTestSuites/Source/UAFTestSuite/Private/UAF/ValueRuntime/Transformers/VectorAttributeTransformerTests.cpp
new file mode 100644
--- /dev/null
+++ b/Engine/Plugins/Experimental/UAF/UAFTestSuites/Source/UAFTestSuite/Private/UAF/ValueRuntime/Transformers/VectorAttributeTransformerTests.cpp
@@ -0,0 +1,305 @@
+// Copyright Epic Games, Inc. All Rights Reserved.
+
+#pragma once
+
+#include "CoreMinimal.h"
+#include "Misc/AutomationTest.h"
+
+#if WITH_DEV_AUTOMATION_TESTS
+
+#include "AnimNextTest.h"
+#include "Animation/BuiltInAttributeTypes.h"
+#include "UAF/AbstractSkeleton/AbstractSkeletonSetBinding.h"
+#include "UAF/Attributes/AttributeBindingDataCache.h"
+#include "UAF/ValueRuntime/ValueRuntimeRegistry.h"
+#include "UAF/ValueRuntime/Transformers/Accumulate.h"
+#include "UAF/ValueRuntime/Transformers/AdditiveSpace.h"
+#include "UAF/ValueRuntime/Transformers/Interpolate.h"
+#include "UAF/ValueRuntime/Transformers/Overwrite.h"
+#include "UAF/ValueRuntime/Transformers/TransformerTestUtil.h"
+#include "UObject/Package.h"
+
+namespace UE::UAF::Tests
+{
+	IMPLEMENT_SIMPLE_AUTOMATION_TEST(FTransformerTests_VectorAttribute, "Animation.UAF.ValueRuntime.Transformers.VectorAttribute", EAutomationTestFlags::EditorContext | EAutomationTestFlags::EngineFilter)
+
+	bool FTransformerTests_VectorAttribute::RunTest(const FString& InParameters)
+	{
+		ON_SCOPE_EXIT{ FUtils::CleanupAfterTests(); };
+
+		TValueOrError<TNonNullPtr<UAbstractSkeletonSetBinding>, FString> SetBindingResult = BuildSetBinding({});
+
+		if (SetBindingResult.HasError())
+		{
+			AddError(SetBindingResult.GetError());
+			return false;
+		}
+		
+		TNonNullPtr<UAbstractSkeletonSetBinding> SetBinding = SetBindingResult.GetValue();
+
+		SetBinding->AddAttributeToSet(FAnimationAttributeIdentifier(TEXT("AttrA"), INDEX_NONE, NAME_None, FVectorAnimationAttribute::StaticStruct()), "SetA");
+		SetBinding->AddAttributeToSet(FAnimationAttributeIdentifier(TEXT("AttrB"), INDEX_NONE, NAME_None, FVectorAnimationAttribute::StaticStruct()), "SetB");
+		SetBinding->AddAttributeToSet(FAnimationAttributeIdentifier(TEXT("AttrC"), INDEX_NONE, NAME_None, FVectorAnimationAttribute::StaticStruct()), "SetC");
+		
+		const FAttributeBindingDataPtr& BindingData = GAttributeBindingDataCache.GetOrAdd(SetBinding, nullptr);
+		FAttributeNamedSetPtr SetA = BindingData->FindNamedSet("SetA");
+		FAttributeNamedSetPtr SetB = BindingData->FindNamedSet("SetB");
+		FAttributeNamedSetPtr SetC = BindingData->FindNamedSet("SetC");
+
+		// Overwrite
+		AddInfo("Overwrite");
+		{
+			// 0.25 weight
+			{
+				FTransformerUnaryTest<FVectorAnimationAttribute, false> VectorTest;
+				VectorTest.AttributeName = "AttrA";
+				VectorTest.Input.Value = FVector(1.0f, 2.0f, 3.0f);
+				VectorTest.ApplyTransformer = [](const FValueBundleHeap& InputValues, FValueBundleHeap& OutputValues)
+					{
+						const auto& TransformerMap = FValueRuntimeRegistry::Get().GetTransformerMap();
+						Transformers::FOverwrite::Apply(TransformerMap, InputValues, 0.25f, OutputValues);
+					};
+				VectorTest.IsExpectedResult = [&](TOptional<FVectorAnimationAttribute> Result) -> TOptional<FString>
+					{
+						if (!Result.IsSet())
+						{
+							return FString("Overwrite - 0.25 weight - expected to produce a value but received none");
+						}
+
+						const FVector ExpectedValue = FVector(1.0f, 2.0f, 3.0f) * 0.25f;
+						if (!ExpectedValue.Equals(Result->Value))
+						{
+							return FString::Format(TEXT("Overwrite - 0.25 weight - expected {0}, Received {1}"), { ExpectedValue.ToString(), Result->Value.ToString() });
+						}
+
+						return TOptional<FString>();
+					};
+			}
+		}
+		
+		// Accumulate
+		AddInfo("Accumulate");
+		{
+			// 0.25 weight
+			{
+				FTransformerBinaryTest<Transformers::FAccumulate, FVectorAnimationAttribute, EInplaceTransformationSupport::Allow> VectorTest;
+				VectorTest.AttributeName = "AttrA";
+				VectorTest.InputA.Value = FVector(1.0f, 2.0f, 3.0f);
+				VectorTest.InputB.Value = FVector(4.0f, 5.0f, 6.0f);
+				VectorTest.ApplyTransformer = [](const FValueBundleHeap& InputAValues, const FValueBundleHeap& InputBValues, FValueBundleHeap& OutputValues)
+					{
+						const auto& TransformerMap = FValueRuntimeRegistry::Get().GetTransformerMap();
+						Transformers::FAccumulate::Apply(TransformerMap, InputAValues, InputBValues, 0.25f, OutputValues);
+					};
+				VectorTest.IsExpectedResult = [&](TOptional<FVectorAnimationAttribute> Result) -> TOptional<FString>
+					{
+						if (!Result.IsSet())
+						{
+							return FString("Accumulate - 0.25 weight - expected to produce a value but received none");
+						}
+						const FVector ExpectedValue = FVector(1.0f, 2.0f, 3.0f) + 0.25f * FVector(4.0f, 5.0f, 6.0f);
+						if (!ExpectedValue.Equals(Result->Value))
+						{
+							return FString::Format(TEXT("Accumulate - 0.25 weight - expected {0}, received {1}"), { ExpectedValue.ToString(), Result->Value.ToString() });
+						}
+						return TOptional<FString>();
+					};
+
+				VectorTest.Run(*this, SetA);
+			}
+		}
+
+		// Interpolate
+		AddInfo("Interpolate");
+		{
+			// 0.25 weight
+			{
+				FTransformerBinaryTest<Transformers::FInterpolate, FVectorAnimationAttribute, EInplaceTransformationSupport::Allow> VectorTest;
+				VectorTest.AttributeName = "AttrA";
+				VectorTest.InputA.Value = FVector(1.0f, 2.0f, 3.0f);
+				VectorTest.InputB.Value = FVector(4.0f, 5.0f, 6.0f);
+				VectorTest.ApplyTransformer = [](const FValueBundleHeap& InputAValues, const FValueBundleHeap& InputBValues, FValueBundleHeap& OutputValues)
+					{
+						const auto& TransformerMap = FValueRuntimeRegistry::Get().GetTransformerMap();
+						Transformers::FInterpolate::Apply(TransformerMap, InputAValues, InputBValues, 0.25f, OutputValues);
+					};
+				VectorTest.IsExpectedResult = [&](TOptional<FVectorAnimationAttribute> Result) -> TOptional<FString>
+					{
+						if (!Result.IsSet())
+						{
+							return FString("Interpolate - 0.25 weight -expected to produce a value but received none");
+						}
+						
+						const FVector ExpectedValue = 0.75f * FVector(1.0f, 2.0f, 3.0f) + 0.25f * FVector(4.0f, 5.0f, 6.0f);
+						if (!ExpectedValue.Equals(Result->Value))
+						{
+							return FString::Format(TEXT("Interpolate - 0.25 weight - expected {0}, received {1}"), { ExpectedValue.ToString(), Result->Value.ToString() });
+						}
+						
+						return TOptional<FString>();
+					};
+
+				VectorTest.Run(*this, SetA);
+			}
+		}
+
+		// Layer
+		AddInfo("Layer");
+		{
+			// Value in base but not in layer
+			{
+				// Weight 0.25
+				{
+					FTransformerBinaryTest<Transformers::FLayer, FVectorAnimationAttribute, EInplaceTransformationSupport::Allow> VectorTest;
+					VectorTest.AttributeName = "AttrA";
+					VectorTest.WeightB = 0.25f;
+					VectorTest.InputA = FVectorAnimationAttribute({ FVector(1.0f, 2.0f, 3.0f) });
+					VectorTest.InputB.Reset();
+					VectorTest.IsExpectedResult = [&](TOptional<FVectorAnimationAttribute> Result) -> TOptional<FString>
+						{
+							if (!Result.IsSet())
+							{
+								return FString("Layer - 0.25 weight, case 1 - expected to produce a value but received none");
+							}
+							
+							const FVector ExpectedValue = FVector(1.0f, 2.0f, 3.0f);
+							if (!ExpectedValue.Equals(Result->Value))
+							{
+								return FString::Format(TEXT("Layer - 0.25 weight, case 1 - expected {0}, received {1}"), { ExpectedValue.ToString(), Result->Value.ToString() });
+							}
+							
+							return TOptional<FString>();
+						};
+
+					VectorTest.Run(*this, SetA, SetB);
+				}
+			}
+
+			// Value in layer but not in base
+			{
+				// Weight 0.25
+				{
+					FTransformerBinaryTest<Transformers::FLayer, FVectorAnimationAttribute, EInplaceTransformationSupport::Allow> VectorTest;
+					VectorTest.AttributeName = "AttrC";
+					VectorTest.WeightB = 0.25f;
+					VectorTest.InputA.Reset();
+					VectorTest.InputB = FVectorAnimationAttribute({ FVector(1.0f, 2.0f, 3.0f) });
+					VectorTest.IsExpectedResult = [&](TOptional<FVectorAnimationAttribute> Result) -> TOptional<FString>
+						{
+							if (Result.IsSet())
+							{
+								return FString("Layer - 0.25 weight, case 2 - expected to not produce a value but received one");
+							}
+														
+							return TOptional<FString>();
+						};
+
+					VectorTest.Run(*this, SetB, SetC);
+				}
+			}
+			
+			// Value in both base and layer
+			{
+				// Weight 0.25
+				{
+					FTransformerBinaryTest<Transformers::FLayer, FVectorAnimationAttribute, EInplaceTransformationSupport::Allow> VectorTest;
+					VectorTest.AttributeName = "AttrB";
+					VectorTest.WeightB = 0.25f;
+					VectorTest.InputA = FVectorAnimationAttribute({ FVector(1.0f, 2.0f, 3.0f) });
+					VectorTest.InputB = FVectorAnimationAttribute({ FVector(4.0f, 5.0f, 6.0f) });
+					VectorTest.IsExpectedResult = [&](TOptional<FVectorAnimationAttribute> Result) -> TOptional<FString>
+						{
+							if (!Result.IsSet())
+							{
+								return FString("Layer - 0.25 weight, case 3 - expected to produce a value but received none");
+							}
+							
+							const FVector ExpectedValue = 0.75f * FVector(1.0f, 2.0f, 3.0f) + 0.25f * FVector(4.0f, 5.0f, 6.0f);
+							if (!ExpectedValue.Equals(Result->Value))
+							{
+								return FString::Format(TEXT("Layer - 0.25 weight, case 3 - expected {0}, received {1}"), { ExpectedValue.ToString(), Result->Value.ToString() });
+							}
+							
+							return TOptional<FString>();
+						};
+
+					VectorTest.Run(*this, SetA, SetB);
+				}
+			}
+		}
+
+		// Make Additive Space
+		AddInfo("MakeAdditiveSpace");
+		{
+			// 0.25 weight
+			{
+				FTransformerBinaryTest<Transformers::FMakeAdditiveSpace, FVectorAnimationAttribute, EInplaceTransformationSupport::Disallow> VectorTest;
+				VectorTest.AttributeName = "AttrA";
+				VectorTest.InputA = FVectorAnimationAttribute({ FVector(1.0f, 2.0f, 3.0f) });
+				VectorTest.InputB = FVectorAnimationAttribute({ FVector(4.0f, 5.0f, 6.0f) });
+				VectorTest.ApplyTransformer = [](const FValueBundleHeap& InputAValues, const FValueBundleHeap& InputBValues, FValueBundleHeap& OutputValues)
+					{
+						const auto& TransformerMap = FValueRuntimeRegistry::Get().GetTransformerMap();
+						Transformers::FMakeAdditiveSpace::Apply(TransformerMap, InputAValues, InputBValues, OutputValues);
+					};
+				VectorTest.IsExpectedResult = [&](TOptional<FVectorAnimationAttribute> Result) -> TOptional<FString>
+					{
+						if (!Result.IsSet())
+						{
+							return FString("MakeAdditiveSpace - expected to produce a value but received none");
+						}
+						
+						const FVector ExpectedValue = FVector(4.0f, 5.0f, 6.0f) - FVector(1.0f, 2.0f, 3.0f);
+						if (!ExpectedValue.Equals(Result->Value))
+						{
+							return FString::Format(TEXT("MakeAdditiveSpace - expected {0}, received {1}"), { ExpectedValue.ToString(), Result->Value.ToString() });
+						}
+						
+						return TOptional<FString>();
+					};
+
+				VectorTest.Run(*this, SetA);
+			}
+		}
+
+		// Apply Additive Space
+		AddInfo("ApplyAdditiveSpace");
+		{
+			// 0.25 weight
+			{
+				FTransformerBinaryTest<Transformers::FMakeAdditiveSpace, FVectorAnimationAttribute, EInplaceTransformationSupport::Allow> VectorTest;
+				VectorTest.AttributeName = "AttrA";
+				VectorTest.InputA = FVectorAnimationAttribute({ FVector(1.0f, 2.0f, 3.0f) });
+				VectorTest.InputB = FVectorAnimationAttribute({ FVector(4.0f, 5.0f, 6.0f) });
+				VectorTest.ApplyTransformer = [](const FValueBundleHeap& InputAValues, const FValueBundleHeap& InputBValues, FValueBundleHeap& OutputValues)
+					{
+						const_cast<FValueBundleHeap&>(InputBValues).SetValueSpace(FValueSpace(EValueSpaceType::Local, true));
+						
+						const auto& TransformerMap = FValueRuntimeRegistry::Get().GetTransformerMap();
+						Transformers::FApplyAdditiveSpace::Apply(TransformerMap, InputAValues, InputBValues, 0.25f, OutputValues);
+					};
+				VectorTest.IsExpectedResult = [&](TOptional<FVectorAnimationAttribute> Result) -> TOptional<FString>
+					{
+						if (!Result.IsSet())
+						{
+							return FString("ApplyAdditiveSpace - 0.25 weight - expected to produce a value but received none");
+						}
+						
+						const FVector ExpectedValue = FVector(1.0f, 2.0f, 3.0f) + 0.25f * FVector(4.0f, 5.0f, 6.0f);
+						if (!ExpectedValue.Equals(Result->Value))
+						{
+							return FString::Format(TEXT("ApplyAdditiveSpace - 0.25 weight - expected {0}, received {1}"), { ExpectedValue.ToString(), Result->Value.ToString() });
+						}
+						
+						return TOptional<FString>();
+					};
+
+				VectorTest.Run(*this, SetA);
+			}
+		}
+		
+		return true;
+	}
+}
+
+#endif // WITH_DEV_AUTOMATION_TESTS
\ No newline at end of file
