Remove pragma once in cpp files

#rnx
#Synced-CL 52003153

diff --git a/Engine/Plugins/Experimental/UAF/UAFTestSuites/Source/UAFTestSuite/Private/UAF/ValueRuntime/Transformers/BoneTransformAttributeTransformerTests.cpp b/Engine/Plugins/Experimental/UAF/UAFTestSuites/Source/UAFTestSuite/Private/UAF/ValueRuntime/Transformers/BoneTransformAttributeTransformerTests.cpp
--- a/Engine/Plugins/Experimental/UAF/UAFTestSuites/Source/UAFTestSuite/Private/UAF/ValueRuntime/Transformers/BoneTransformAttributeTransformerTests.cpp
+++ b/Engine/Plugins/Experimental/UAF/UAFTestSuites/Source/UAFTestSuite/Private/UAF/ValueRuntime/Transformers/BoneTransformAttributeTransformerTests.cpp
@@ -1,7 +1,5 @@
 // Copyright Epic Games, Inc. All Rights Reserved.
 
-#pragma once
-
 #include "CoreMinimal.h"
 #include "Misc/AutomationTest.h"
 
diff --git a/Engine/Plugins/Experimental/UAF/UAFTestSuites/Source/UAFTestSuite/Private/UAF/ValueRuntime/Transformers/FloatAttributeTransformerTests.cpp b/Engine/Plugins/Experimental/UAF/UAFTestSuites/Source/UAFTestSuite/Private/UAF/ValueRuntime/Transformers/FloatAttributeTransformerTests.cpp
--- a/Engine/Plugins/Experimental/UAF/UAFTestSuites/Source/UAFTestSuite/Private/UAF/ValueRuntime/Transformers/FloatAttributeTransformerTests.cpp
+++ b/Engine/Plugins/Experimental/UAF/UAFTestSuites/Source/UAFTestSuite/Private/UAF/ValueRuntime/Transformers/FloatAttributeTransformerTests.cpp
@@ -1,7 +1,5 @@
 // Copyright Epic Games, Inc. All Rights Reserved.
 
-#pragma once
-
 #include "CoreMinimal.h"
 #include "Misc/AutomationTest.h"
 
diff --git a/Engine/Plugins/Experimental/UAF/UAFTestSuites/Source/UAFTestSuite/Private/UAF/ValueRuntime/Transformers/IntegerAttributeTransformerTests.cpp b/Engine/Plugins/Experimental/UAF/UAFTestSuites/Source/UAFTestSuite/Private/UAF/ValueRuntime/Transformers/IntegerAttributeTransformerTests.cpp
--- a/Engine/Plugins/Experimental/UAF/UAFTestSuites/Source/UAFTestSuite/Private/UAF/ValueRuntime/Transformers/IntegerAttributeTransformerTests.cpp
+++ b/Engine/Plugins/Experimental/UAF/UAFTestSuites/Source/UAFTestSuite/Private/UAF/ValueRuntime/Transformers/IntegerAttributeTransformerTests.cpp
@@ -1,7 +1,5 @@
 // Copyright Epic Games, Inc. All Rights Reserved.
 
-#pragma once
-
 #include "CoreMinimal.h"
 #include "Misc/AutomationTest.h"
 
diff --git a/Engine/Plugins/Experimental/UAF/UAFTestSuites/Source/UAFTestSuite/Private/UAF/ValueRuntime/Transformers/QuaternionAttributeTransformerTests.cpp b/Engine/Plugins/Experimental/UAF/UAFTestSuites/Source/UAFTestSuite/Private/UAF/ValueRuntime/Transformers/QuaternionAttributeTransformerTests.cpp
--- a/Engine/Plugins/Experimental/UAF/UAFTestSuites/Source/UAFTestSuite/Private/UAF/ValueRuntime/Transformers/QuaternionAttributeTransformerTests.cpp
+++ b/Engine/Plugins/Experimental/UAF/UAFTestSuites/Source/UAFTestSuite/Private/UAF/ValueRuntime/Transformers/QuaternionAttributeTransformerTests.cpp
@@ -1,7 +1,5 @@
 // Copyright Epic Games, Inc. All Rights Reserved.
 
-#pragma once
-
 #include "CoreMinimal.h"
 #include "Misc/AutomationTest.h"
 
diff --git a/Engine/Plugins/Experimental/UAF/UAFTestSuites/Source/UAFTestSuite/Private/UAF/ValueRuntime/Transformers/TransformAttributeTransformerTests.cpp b/Engine/Plugins/Experimental/UAF/UAFTestSuites/Source/UAFTestSuite/Private/UAF/ValueRuntime/Transformers/TransformAttributeTransformerTests.cpp
--- a/Engine/Plugins/Experimental/UAF/UAFTestSuites/Source/UAFTestSuite/Private/UAF/ValueRuntime/Transformers/TransformAttributeTransformerTests.cpp
+++ b/Engine/Plugins/Experimental/UAF/UAFTestSuites/Source/UAFTestSuite/Private/UAF/ValueRuntime/Transformers/TransformAttributeTransformerTests.cpp
@@ -1,7 +1,5 @@
 // Copyright Epic Games, Inc. All Rights Reserved.
 
-#pragma once
-
 #include "CoreMinimal.h"
 #include "Misc/AutomationTest.h"
 
diff --git a/Engine/Plugins/Experimental/UAF/UAFTestSuites/Source/UAFTestSuite/Private/UAF/ValueRuntime/Transformers/VectorAttributeTransformerTests.cpp b/Engine/Plugins/Experimental/UAF/UAFTestSuites/Source/UAFTestSuite/Private/UAF/ValueRuntime/Transformers/VectorAttributeTransformerTests.cpp
--- a/Engine/Plugins/Experimental/UAF/UAFTestSuites/Source/UAFTestSuite/Private/UAF/ValueRuntime/Transformers/VectorAttributeTransformerTests.cpp
+++ b/Engine/Plugins/Experimental/UAF/UAFTestSuites/Source/UAFTestSuite/Private/UAF/ValueRuntime/Transformers/VectorAttributeTransformerTests.cpp
@@ -1,7 +1,5 @@
 // Copyright Epic Games, Inc. All Rights Reserved.
 
-#pragma once
-
 #include "CoreMinimal.h"
 #include "Misc/AutomationTest.h"
 
