Hide FUAFTestAnimNodeData: it is just used for local development

#Synced-CL 52150898

diff --git a/Engine/Plugins/Experimental/UAF/UAFTestSuites/Source/UAFAnimNodeTestData/Public/UAFAnimNodeTestVars.h b/Engine/Plugins/Experimental/UAF/UAFTestSuites/Source/UAFAnimNodeTestData/Public/UAFAnimNodeTestVars.h
--- a/Engine/Plugins/Experimental/UAF/UAFTestSuites/Source/UAFAnimNodeTestData/Public/UAFAnimNodeTestVars.h
+++ b/Engine/Plugins/Experimental/UAF/UAFTestSuites/Source/UAFAnimNodeTestData/Public/UAFAnimNodeTestVars.h
@@ -97,7 +97,7 @@ enum class EUAFAnimNodeTestEnum : uint8
 	Gamma = 2,
 };
 
-USTRUCT(DisplayName = "Test Anim Node")
+USTRUCT(DisplayName = "Test Anim Node", meta=(Hidden))
 struct FUAFTestAnimNodeData : public FUAFAnimNodeData
 {
 	GENERATED_BODY()
