UAF: Added additional tests to cover GetRegisteredObjectClasses behavior for UAFAssetData

#Synced-CL 52246754

diff --git a/Engine/Plugins/Experimental/UAF/UAFTestSuites/Source/UAFCQTestSuite/Private/UAFAssetDataTests.cpp b/Engine/Plugins/Experimental/UAF/UAFTestSuites/Source/UAFCQTestSuite/Private/UAFAssetDataTests.cpp
--- a/Engine/Plugins/Experimental/UAF/UAFTestSuites/Source/UAFCQTestSuite/Private/UAFAssetDataTests.cpp
+++ b/Engine/Plugins/Experimental/UAF/UAFTestSuites/Source/UAFCQTestSuite/Private/UAFAssetDataTests.cpp
@@ -485,6 +485,79 @@ TEST_CLASS(FUAFAssetDataTests, "Animation.UAF.Functional")
 			Graph = nullptr;
     	});
     }
+	
+	FTopLevelAssetPath TestObjectClassPath;
+	FTopLevelAssetPath ChildASubobjectClassPath;
+	TEST_METHOD(Verify_GetRegisteredObjectClasses_By_Base_Asset_Type)
+	{
+		TestCommandBuilder
+		.Do(TEXT("Register test types"), [&]()
+			{
+				/** 
+				 * Class Hierarchy for Test
+				 * UUAFAssetDataTestObject (registered) -> UUAFAssetDataTestObject_ChildA (registered) -> UUAFAssetDataTestObject_GrandChildA
+				 *										-> UUAFAssetDataTestObject_ChildB -> UUAFAssetDataTestObject_GrandChildB
+				*/
+				
+				// Register UUAFAssetDataTestObject with FUAFGraphFactoryAsset_Test
+				TestObjectClassPath = UE::UAF::FAssetDataFactory::RegisterAsset<FUAFGraphFactoryAsset_Test, UUAFAssetDataTestObject>(
+					[](const UUAFAssetDataTestObject*) { return FUAFGraphFactoryAsset_Test(); });
+
+				// Register UUAFAssetDataTestObject_ChildA with its own factory FUAFGraphFactoryAsset_TestChildA
+				ChildASubobjectClassPath = UE::UAF::FAssetDataFactory::RegisterAsset<FUAFGraphFactoryAsset_TestChildA, UUAFAssetDataTestObject_ChildA>(
+					[](const UUAFAssetDataTestObject_ChildA*) { return FUAFGraphFactoryAsset_TestChildA(); });
+			})
+		.Then(TEXT("Verify classes returned with GetRegisteredObjectClasses for FUAFGraphFactoryAsset_Test"), [&]()
+			{
+				TArray<UClass*> Classes = UE::UAF::FAssetDataFactory::GetRegisteredObjectClasses<FUAFGraphFactoryAsset_Test>();
+
+				// The directly registered class must be present
+				ASSERT_THAT(IsTrue(Classes.Contains(UUAFAssetDataTestObject::StaticClass())))
+
+				// Any unregistered subclass is gathered via GatherDerivedObjectClasses
+				ASSERT_THAT(IsTrue(Classes.Contains(UUAFAssetDataTestObject_ChildB::StaticClass())))
+				ASSERT_THAT(IsTrue(Classes.Contains(UUAFAssetDataTestObject_GrandChildB::StaticClass())))
+
+				// A subclass that has its own separate registration must NOT be gathered —
+				// GatherDerivedObjectClasses stops at classes that have their own initializer
+				ASSERT_THAT(IsFalse(Classes.Contains(UUAFAssetDataTestObject_ChildA::StaticClass())))
+				
+				// A subclass whose parent has its own initializer must not be gathered
+				ASSERT_THAT(IsFalse(Classes.Contains(UUAFAssetDataTestObject_GrandChildA::StaticClass())))
+			})
+		.Then(TEXT("Verify classes returned for FUAFGraphFactoryAsset_TestChildA"), [&]()
+			{
+				TArray<UClass*> Classes = UE::UAF::FAssetDataFactory::GetRegisteredObjectClasses<FUAFGraphFactoryAsset_TestChildA>();
+
+				// The class registered under this asset type should appear
+				ASSERT_THAT(IsTrue(Classes.Contains(UUAFAssetDataTestObject_ChildA::StaticClass())))
+				
+				// An unregistered subclass is gathered via GatherDerivedObjectClasses
+				ASSERT_THAT(IsTrue(Classes.Contains(UUAFAssetDataTestObject_GrandChildA::StaticClass())))
+
+				// The parent class and its unregistered subclass belong to a different asset type
+				ASSERT_THAT(IsFalse(Classes.Contains(UUAFAssetDataTestObject::StaticClass())))
+				ASSERT_THAT(IsFalse(Classes.Contains(UUAFAssetDataTestObject_ChildB::StaticClass())))
+				
+			})
+		.Then(TEXT("Verify classes returned for FUAFGraphFactoryAsset (common base)"), [&]()
+			{
+				TArray<UClass*> Classes = UE::UAF::FAssetDataFactory::GetRegisteredObjectClasses<FUAFGraphFactoryAsset>();
+
+				// All subclasses should appear when querying the shared base asset type
+				ASSERT_THAT(IsTrue(Classes.Contains(UUAFAssetDataTestObject::StaticClass())))
+				ASSERT_THAT(IsTrue(Classes.Contains(UUAFAssetDataTestObject_ChildA::StaticClass())))
+				ASSERT_THAT(IsTrue(Classes.Contains(UUAFAssetDataTestObject_ChildB::StaticClass())))
+				ASSERT_THAT(IsTrue(Classes.Contains(UUAFAssetDataTestObject_GrandChildA::StaticClass())))
+				ASSERT_THAT(IsTrue(Classes.Contains(UUAFAssetDataTestObject_GrandChildB::StaticClass())))
+			})
+		.OnTearDown(TEXT("Unregister test types"), [&]()
+			{
+				UE::UAF::FAssetDataFactory::UnregisterAsset(TestObjectClassPath);
+				UE::UAF::FAssetDataFactory::UnregisterAsset(ChildASubobjectClassPath);
+			});
+	}
+
 };
 
 TObjectPtr<USkeletalMesh> FUAFAssetDataTests::GSkeletalMesh = nullptr;
diff --git a/Engine/Plugins/Experimental/UAF/UAFTestSuites/Source/UAFCQTestSuite/Private/UAFAssetDataTests.h b/Engine/Plugins/Experimental/UAF/UAFTestSuites/Source/UAFCQTestSuite/Private/UAFAssetDataTests.h
--- a/Engine/Plugins/Experimental/UAF/UAFTestSuites/Source/UAFCQTestSuite/Private/UAFAssetDataTests.h
+++ b/Engine/Plugins/Experimental/UAF/UAFTestSuites/Source/UAFCQTestSuite/Private/UAFAssetDataTests.h
@@ -24,3 +24,46 @@ public:
 	int TestData = 0;
 };
 
+
+// Subclass of UUAFAssetDataTestObject that IS registered with its own asset type (FUAFGraphFactoryAsset_TestChildA).
+// GetRegisteredObjectClasses with the parent's asset type should NOT include it because it stops recursing at classes that have their own registered initializer.
+UCLASS()
+class UUAFAssetDataTestObject_ChildA : public UUAFAssetDataTestObject
+{
+	GENERATED_BODY()
+};
+
+// Factory Asset to register the asset type of UUAFAssetDataTestObject_ChildA 
+USTRUCT(meta=(Hidden))
+struct FUAFGraphFactoryAsset_TestChildA : public FUAFGraphFactoryAsset
+{
+	GENERATED_BODY()
+};
+
+// Subclass of UUAFAssetDataTestObject that is NOT registered with FAssetDataFactory.
+// GetRegisteredObjectClasses with the parent's asset type should include it
+UCLASS()
+class UUAFAssetDataTestObject_ChildB : public UUAFAssetDataTestObject
+{
+	GENERATED_BODY()
+};
+
+// Subclass of UUAFAssetDataTestObject_ChildA
+// GetRegisteredObjectClasses with UUAFAssetDataTestObject_ChildA as base type should include this 
+// GetRegisteredObjectClasses with UUAFAssetDataTestObject should not include this 
+UCLASS()
+class UUAFAssetDataTestObject_GrandChildA : public UUAFAssetDataTestObject_ChildA
+{
+	GENERATED_BODY()
+};
+ 
+// Subclass of UUAFAssetDataTestObject_ChildB
+// GetRegisteredObjectClasses with UUAFAssetDataTestObject should include this and its parent 
+UCLASS()
+class UUAFAssetDataTestObject_GrandChildB : public UUAFAssetDataTestObject_ChildB
+{
+	GENERATED_BODY()
+};
+ 
+
+
