Rename GetComponent to GetOrAddComponent to match functionality

#rnx
#Synced-CL 52570859

diff --git a/Engine/Plugins/Experimental/UAF/UAF/Source/UAF/Private/Graph/RigUnit_AnimNextWriteSkeletalMeshComponentPose.cpp b/Engine/Plugins/Experimental/UAF/UAF/Source/UAF/Private/Graph/RigUnit_AnimNextWriteSkeletalMeshComponentPose.cpp
--- a/Engine/Plugins/Experimental/UAF/UAF/Source/UAF/Private/Graph/RigUnit_AnimNextWriteSkeletalMeshComponentPose.cpp
+++ b/Engine/Plugins/Experimental/UAF/UAF/Source/UAF/Private/Graph/RigUnit_AnimNextWriteSkeletalMeshComponentPose.cpp
@@ -27,7 +27,7 @@ FRigUnit_AnimNextWriteSkeletalMeshComponentPose_Execute()
 	{
 		const FAnimNextModuleContextData& ModuleContextData = ExecuteContext.GetContextData<FAnimNextModuleContextData>();
 		FAnimNextModuleInstance& ModuleInstance = ModuleContextData.GetModuleInstance();
-		const FAnimNextSkeletalMeshComponentReferenceComponent& ComponentReference = ModuleInstance.GetComponent<FAnimNextSkeletalMeshComponentReferenceComponent>();
+		const FAnimNextSkeletalMeshComponentReferenceComponent& ComponentReference = ModuleInstance.GetOrAddComponent<FAnimNextSkeletalMeshComponentReferenceComponent>();
 		OutputComponent = ComponentReference.GetComponent();
 	}
 
diff --git a/Engine/Plugins/Experimental/UAF/UAF/Source/UAF/Private/Graph/RigUnit_GetActorTransform.cpp b/Engine/Plugins/Experimental/UAF/UAF/Source/UAF/Private/Graph/RigUnit_GetActorTransform.cpp
--- a/Engine/Plugins/Experimental/UAF/UAF/Source/UAF/Private/Graph/RigUnit_GetActorTransform.cpp
+++ b/Engine/Plugins/Experimental/UAF/UAF/Source/UAF/Private/Graph/RigUnit_GetActorTransform.cpp
@@ -77,6 +77,6 @@ FRigUnit_GetActorTransform_Execute()
 		UAF_RIGUNIT_LOG(Warning, TEXT("Could not get actor transform - Root module instance is not valid."));
 		return;
     }
-	const FAnimNextActorTransformComponent& TransformComponent = AssetInstance.GetRootInstance()->GetComponent<FAnimNextActorTransformComponent>();
+	const FAnimNextActorTransformComponent& TransformComponent = AssetInstance.GetRootInstance()->GetOrAddComponent<FAnimNextActorTransformComponent>();
 	Transform = TransformComponent.GetActorTransform();
 }
diff --git a/Engine/Plugins/Experimental/UAF/UAF/Source/UAF/Private/Graph/RigUnit_MakeReferencePoseFromSkeletalMeshComponent.cpp b/Engine/Plugins/Experimental/UAF/UAF/Source/UAF/Private/Graph/RigUnit_MakeReferencePoseFromSkeletalMeshComponent.cpp
--- a/Engine/Plugins/Experimental/UAF/UAF/Source/UAF/Private/Graph/RigUnit_MakeReferencePoseFromSkeletalMeshComponent.cpp
+++ b/Engine/Plugins/Experimental/UAF/UAF/Source/UAF/Private/Graph/RigUnit_MakeReferencePoseFromSkeletalMeshComponent.cpp
@@ -27,7 +27,7 @@ FRigUnit_MakeReferencePoseFromSkeletalMeshComponent_Execute()
 	{
 		const FAnimNextModuleContextData& ModuleContextData = ExecuteContext.GetContextData<FAnimNextModuleContextData>();
 		FAnimNextModuleInstance& ModuleInstance = ModuleContextData.GetModuleInstance();
-		const FAnimNextSkeletalMeshComponentReferenceComponent& ComponentReference = ModuleInstance.GetComponent<FAnimNextSkeletalMeshComponentReferenceComponent>();
+		const FAnimNextSkeletalMeshComponentReferenceComponent& ComponentReference = ModuleInstance.GetOrAddComponent<FAnimNextSkeletalMeshComponentReferenceComponent>();
 		InputComponent = ComponentReference.GetComponent();
 	}
 
diff --git a/Engine/Plugins/Experimental/UAF/UAF/Source/UAF/Private/Graph/RigUnit_UAFWriteSystemOutput.cpp b/Engine/Plugins/Experimental/UAF/UAF/Source/UAF/Private/Graph/RigUnit_UAFWriteSystemOutput.cpp
--- a/Engine/Plugins/Experimental/UAF/UAF/Source/UAF/Private/Graph/RigUnit_UAFWriteSystemOutput.cpp
+++ b/Engine/Plugins/Experimental/UAF/UAF/Source/UAF/Private/Graph/RigUnit_UAFWriteSystemOutput.cpp
@@ -16,6 +16,6 @@ FRigUnit_UAFWriteSystemOutput_Execute()
 
 	const FAnimNextModuleContextData& ModuleContextData = ExecuteContext.GetContextData<FAnimNextModuleContextData>();
 	FAnimNextModuleInstance& ModuleInstance = ModuleContextData.GetModuleInstance();
-	FUAFSystemOutputComponent& OutputPoseComponent = ModuleInstance.GetComponent<FUAFSystemOutputComponent>();
+	FUAFSystemOutputComponent& OutputPoseComponent = ModuleInstance.GetOrAddComponent<FUAFSystemOutputComponent>();
 	OutputPoseComponent.WriteOutput(Value);
 }
diff --git a/Engine/Plugins/Experimental/UAF/UAF/Source/UAF/Private/Module/ModuleTaskContext.cpp b/Engine/Plugins/Experimental/UAF/UAF/Source/UAF/Private/Module/ModuleTaskContext.cpp
--- a/Engine/Plugins/Experimental/UAF/UAF/Source/UAF/Private/Module/ModuleTaskContext.cpp
+++ b/Engine/Plugins/Experimental/UAF/UAF/Source/UAF/Private/Module/ModuleTaskContext.cpp
@@ -36,7 +36,7 @@ void FModuleTaskContext::TryAccessComponent(UScriptStruct* InComponentType, TFun
 
 void FModuleTaskContext::AccessComponent(UScriptStruct* InComponentType, TFunctionRef<void(FUAFAssetInstanceComponent&)> InFunction) const
 {
-	FUAFAssetInstanceComponent& Component = ModuleInstance.GetComponent(InComponentType);
+	FUAFAssetInstanceComponent& Component = ModuleInstance.GetOrAddComponent(InComponentType);
 	InFunction(Component);
 }
 
diff --git a/Engine/Plugins/Experimental/UAF/UAF/Source/UAF/Private/Module/UAFNotifyDispatcherComponent.cpp b/Engine/Plugins/Experimental/UAF/UAF/Source/UAF/Private/Module/UAFNotifyDispatcherComponent.cpp
--- a/Engine/Plugins/Experimental/UAF/UAF/Source/UAF/Private/Module/UAFNotifyDispatcherComponent.cpp
+++ b/Engine/Plugins/Experimental/UAF/UAF/Source/UAF/Private/Module/UAFNotifyDispatcherComponent.cpp
@@ -14,7 +14,7 @@ void FUAFNotifyDispatcherComponent::OnBindToInstance()
 {
 	// Default to the current host's skeletal mesh component, if any
 	FAnimNextModuleInstance& ModuleInstance = GetModuleInstance();
-	const FAnimNextSkeletalMeshComponentReferenceComponent& ComponentReference = ModuleInstance.GetComponent<FAnimNextSkeletalMeshComponentReferenceComponent>();
+	const FAnimNextSkeletalMeshComponentReferenceComponent& ComponentReference = ModuleInstance.GetOrAddComponent<FAnimNextSkeletalMeshComponentReferenceComponent>();
 	SkeletalMeshComponent = ComponentReference.GetComponent();
 	NotifyQueue.PredictedLODLevel = SkeletalMeshComponent ? SkeletalMeshComponent->GetPredictedLODLevel() : 0;
 }
diff --git a/Engine/Plugins/Experimental/UAF/UAF/Source/UAF/Public/UAFAssetInstance.h b/Engine/Plugins/Experimental/UAF/UAF/Source/UAF/Public/UAFAssetInstance.h
--- a/Engine/Plugins/Experimental/UAF/UAF/Source/UAF/Public/UAFAssetInstance.h
+++ b/Engine/Plugins/Experimental/UAF/UAF/Source/UAF/Public/UAFAssetInstance.h
@@ -171,7 +171,7 @@ struct FUAFAssetInstance
 	// Returns a typed instance component, creating it lazily the first time it is queried.
 	// When querying for structs of a base type, existing derived types will be returned, rather than a new base struct created.
 	template<typename ComponentType>
-	[[nodiscard]] ComponentType& GetComponent()
+	[[nodiscard]] ComponentType& GetOrAddComponent()
 	{
 		static_assert(std::is_base_of<FUAFAssetInstanceComponent, ComponentType>::value, "ComponentType type must derive from FUAFAssetInstanceComponent");
 
@@ -191,7 +191,7 @@ struct FUAFAssetInstance
 
 	// Returns an instance component of the requested type, creating it lazily the first time it is queried.
 	// When querying for structs of a base type, existing derived types will be returned, rather than a new base struct created.
-	[[nodiscard]] FUAFAssetInstanceComponent& GetComponent(TSubScriptStructOf<FUAFAssetInstanceComponent> InScriptStruct)
+	[[nodiscard]] FUAFAssetInstanceComponent& GetOrAddComponent(TSubScriptStructOf<FUAFAssetInstanceComponent> InScriptStruct)
 	{
 		check(InScriptStruct.Get() != nullptr);
 		if (FUAFAssetInstanceComponent* ExistingComponent = TryGetComponent(InScriptStruct))
diff --git a/Engine/Plugins/Experimental/UAF/UAFAnimGraph/Source/UAFAnimGraph/Private/Graph/AnimNextGraphInstance.cpp b/Engine/Plugins/Experimental/UAF/UAFAnimGraph/Source/UAFAnimGraph/Private/Graph/AnimNextGraphInstance.cpp
--- a/Engine/Plugins/Experimental/UAF/UAFAnimGraph/Source/UAFAnimGraph/Private/Graph/AnimNextGraphInstance.cpp
+++ b/Engine/Plugins/Experimental/UAF/UAFAnimGraph/Source/UAFAnimGraph/Private/Graph/AnimNextGraphInstance.cpp
@@ -362,7 +362,7 @@ void FAnimNextGraphInstance::RunOrQueueTask(UE::UAF::FUniqueAnimGraphTask&& InTa
 
 	if (AllocationState == EAllocationState::PreAllocated)
 	{
-		FUAFAnimGraphInitializerComponent& InitializerComponent = GetComponent<FUAFAnimGraphInitializerComponent>();
+		FUAFAnimGraphInitializerComponent& InitializerComponent = GetOrAddComponent<FUAFAnimGraphInitializerComponent>();
 		InitializerComponent.TaskQueue.Add(MoveTemp(InTaskFunction));
 	}
 	else
diff --git a/Engine/Plugins/Experimental/UAF/UAFAnimGraph/Source/UAFAnimGraph/Private/Graph/RigUnit_AnimNextRunAnimationGraph_v1.cpp b/Engine/Plugins/Experimental/UAF/UAFAnimGraph/Source/UAFAnimGraph/Private/Graph/RigUnit_AnimNextRunAnimationGraph_v1.cpp
--- a/Engine/Plugins/Experimental/UAF/UAFAnimGraph/Source/UAFAnimGraph/Private/Graph/RigUnit_AnimNextRunAnimationGraph_v1.cpp
+++ b/Engine/Plugins/Experimental/UAF/UAFAnimGraph/Source/UAFAnimGraph/Private/Graph/RigUnit_AnimNextRunAnimationGraph_v1.cpp
@@ -31,7 +31,7 @@ FRigUnit_AnimNextRunAnimationGraph_v1_Execute()
 		return;
 	}
 
-	FAnimNextModuleAnimGraphComponent& AnimationGraphComponent = ModuleInstance.GetComponent<FAnimNextModuleAnimGraphComponent>();
+	FAnimNextModuleAnimGraphComponent& AnimationGraphComponent = ModuleInstance.GetOrAddComponent<FAnimNextModuleAnimGraphComponent>();
 
 	if(Graph == nullptr)
 	{
@@ -102,7 +102,7 @@ FRigUnit_AnimNextRunAnimationGraph_v1_Execute()
 
 	if (UseAbstractSkeletonRuntime())
 	{
-		FUAFSystemOutputComponent& OutputPoseComponent = ModuleInstance.GetComponent<FUAFSystemOutputComponent>();
+		FUAFSystemOutputComponent& OutputPoseComponent = ModuleInstance.GetOrAddComponent<FUAFSystemOutputComponent>();
     		
     	USkeletalMesh* SkeletalMesh = OutputPoseComponent.GetSkeletalMesh();
     	check(SkeletalMesh);
diff --git a/Engine/Plugins/Experimental/UAF/UAFAnimGraph/Source/UAFAnimGraph/Private/Graph/RigUnit_AnimNextRunAnimationGraph_v2.cpp b/Engine/Plugins/Experimental/UAF/UAFAnimGraph/Source/UAFAnimGraph/Private/Graph/RigUnit_AnimNextRunAnimationGraph_v2.cpp
--- a/Engine/Plugins/Experimental/UAF/UAFAnimGraph/Source/UAFAnimGraph/Private/Graph/RigUnit_AnimNextRunAnimationGraph_v2.cpp
+++ b/Engine/Plugins/Experimental/UAF/UAFAnimGraph/Source/UAFAnimGraph/Private/Graph/RigUnit_AnimNextRunAnimationGraph_v2.cpp
@@ -36,7 +36,7 @@ FRigUnit_AnimNextRunAnimationGraph_v2_Execute()
 	// Fall back to output pose ref pose
 	if (!RefPoseDataHandle.IsValid())
 	{
-		FUAFSystemOutputComponent& OutputPoseComponent = ModuleInstance.GetComponent<FUAFSystemOutputComponent>();
+		FUAFSystemOutputComponent& OutputPoseComponent = ModuleInstance.GetOrAddComponent<FUAFSystemOutputComponent>();
 		RefPoseDataHandle = OutputPoseComponent.GetRefPose().ReferencePose;
 	}
 
@@ -65,7 +65,7 @@ FRigUnit_AnimNextRunAnimationGraph_v2_Execute()
 
 	// Get a host to run this graph
 	const UUAFAnimGraph* HostGraph = Graph.HostGraph ? Graph.HostGraph.Get() : FAnimGraphFactory::GetDefaultGraphHost();
-	FAnimNextModuleAnimGraphComponent& AnimationGraphComponent = ModuleInstance.GetComponent<FAnimNextModuleAnimGraphComponent>();
+	FAnimNextModuleAnimGraphComponent& AnimationGraphComponent = ModuleInstance.GetOrAddComponent<FAnimNextModuleAnimGraphComponent>();
 	if(HostGraph == nullptr)
 	{
 		AnimationGraphComponent.ReleaseInstance(WorkData.WeakHost);
@@ -92,7 +92,7 @@ FRigUnit_AnimNextRunAnimationGraph_v2_Execute()
 		{
 			WorkData.WeakHost = NewGraphInstance;
 
-			FGraphInstanceInjectionComponent& HostInjectionComponent = NewGraphInstance->GetComponent<FGraphInstanceInjectionComponent>();
+			FGraphInstanceInjectionComponent& HostInjectionComponent = NewGraphInstance->GetOrAddComponent<FGraphInstanceInjectionComponent>();
 			WorkData.InjectedGraphReference = HostInjectionComponent.GetInjectionInfo().GetDefaultInjectionSite();
 		}
 	}
@@ -142,7 +142,7 @@ FRigUnit_AnimNextRunAnimationGraph_v2_Execute()
 
 	if (UseAbstractSkeletonRuntime())
 	{
-		FUAFSystemOutputComponent& OutputPoseComponent = ModuleInstance.GetComponent<FUAFSystemOutputComponent>();
+		FUAFSystemOutputComponent& OutputPoseComponent = ModuleInstance.GetOrAddComponent<FUAFSystemOutputComponent>();
 		
 		USkeletalMesh* SkeletalMesh = OutputPoseComponent.GetSkeletalMesh();
 		
diff --git a/Engine/Plugins/Experimental/UAF/UAFAnimGraph/Source/UAFAnimGraph/Private/Graph/RigUnit_AnimNextSetNotifyContext.cpp b/Engine/Plugins/Experimental/UAF/UAFAnimGraph/Source/UAFAnimGraph/Private/Graph/RigUnit_AnimNextSetNotifyContext.cpp
--- a/Engine/Plugins/Experimental/UAF/UAFAnimGraph/Source/UAFAnimGraph/Private/Graph/RigUnit_AnimNextSetNotifyContext.cpp
+++ b/Engine/Plugins/Experimental/UAF/UAFAnimGraph/Source/UAFAnimGraph/Private/Graph/RigUnit_AnimNextSetNotifyContext.cpp
@@ -11,7 +11,7 @@
 FRigUnit_AnimNextSetNotifyContext_Execute()
 {
 	FAnimNextModuleInstance& ModuleInstance = ExecuteContext.GetContextData<FAnimNextModuleContextData>().GetModuleInstance();
-	FUAFNotifyDispatcherComponent& NotifyDispatcher = ModuleInstance.GetComponent<FUAFNotifyDispatcherComponent>();
+	FUAFNotifyDispatcherComponent& NotifyDispatcher = ModuleInstance.GetOrAddComponent<FUAFNotifyDispatcherComponent>();
 	NotifyDispatcher.SkeletalMeshComponent = SkeletalMeshComponent;
 	NotifyDispatcher.NotifyQueue.PredictedLODLevel = SkeletalMeshComponent ? SkeletalMeshComponent->GetPredictedLODLevel() : 0;
 }
diff --git a/Engine/Plugins/Experimental/UAF/UAFAnimGraph/Source/UAFAnimGraph/Private/Graph/RigUnit_GetPostProcessAnimation.cpp b/Engine/Plugins/Experimental/UAF/UAFAnimGraph/Source/UAFAnimGraph/Private/Graph/RigUnit_GetPostProcessAnimation.cpp
--- a/Engine/Plugins/Experimental/UAF/UAFAnimGraph/Source/UAFAnimGraph/Private/Graph/RigUnit_GetPostProcessAnimation.cpp
+++ b/Engine/Plugins/Experimental/UAF/UAFAnimGraph/Source/UAFAnimGraph/Private/Graph/RigUnit_GetPostProcessAnimation.cpp
@@ -22,7 +22,7 @@ FRigUnit_GetPostProcessAnimation_Execute()
 	{
 		const FAnimNextModuleContextData& ModuleContextData = ExecuteContext.GetContextData<FAnimNextModuleContextData>();
 		FAnimNextModuleInstance& ModuleInstance = ModuleContextData.GetModuleInstance();
-		const FAnimNextSkeletalMeshComponentReferenceComponent& ComponentReference = ModuleInstance.GetComponent<FAnimNextSkeletalMeshComponentReferenceComponent>();
+		const FAnimNextSkeletalMeshComponentReferenceComponent& ComponentReference = ModuleInstance.GetOrAddComponent<FAnimNextSkeletalMeshComponentReferenceComponent>();
 		UsedSkeletalMeshComponent = ComponentReference.GetComponent();
 	}
 
diff --git a/Engine/Plugins/Experimental/UAF/UAFAnimGraph/Source/UAFAnimGraph/Private/Graph/RigUnit_UAFRunAsset.cpp b/Engine/Plugins/Experimental/UAF/UAFAnimGraph/Source/UAFAnimGraph/Private/Graph/RigUnit_UAFRunAsset.cpp
--- a/Engine/Plugins/Experimental/UAF/UAFAnimGraph/Source/UAFAnimGraph/Private/Graph/RigUnit_UAFRunAsset.cpp
+++ b/Engine/Plugins/Experimental/UAF/UAFAnimGraph/Source/UAFAnimGraph/Private/Graph/RigUnit_UAFRunAsset.cpp
@@ -30,7 +30,7 @@ FRigUnit_UAFRunAsset_Execute()
 	const FAnimNextModuleContextData& ModuleContextData = ExecuteContext.GetContextData<FAnimNextModuleContextData>();
 	FAnimNextModuleInstance& ModuleInstance = ModuleContextData.GetModuleInstance();
 
-	FUAFSystemOutputComponent& OutputPoseComponent = ModuleInstance.GetComponent<FUAFSystemOutputComponent>();
+	FUAFSystemOutputComponent& OutputPoseComponent = ModuleInstance.GetOrAddComponent<FUAFSystemOutputComponent>();
 	FDataHandle RefPoseDataHandle = OutputPoseComponent.GetRefPose().ReferencePose;
 	if (!RefPoseDataHandle.IsValid())
 	{
@@ -50,7 +50,7 @@ FRigUnit_UAFRunAsset_Execute()
 
 	// Get a host to run this graph
 	const UUAFAnimGraph* HostGraph = Graph.HostGraph ? Graph.HostGraph.Get() : FAnimGraphFactory::GetDefaultGraphHost();
-	FAnimNextModuleAnimGraphComponent& AnimationGraphComponent = ModuleInstance.GetComponent<FAnimNextModuleAnimGraphComponent>();
+	FAnimNextModuleAnimGraphComponent& AnimationGraphComponent = ModuleInstance.GetOrAddComponent<FAnimNextModuleAnimGraphComponent>();
 	if(HostGraph == nullptr)
 	{
 		AnimationGraphComponent.ReleaseInstance(WorkData.WeakHost);
@@ -76,7 +76,7 @@ FRigUnit_UAFRunAsset_Execute()
 		{
 			WorkData.WeakHost = NewGraphInstance;
 
-			FGraphInstanceInjectionComponent& HostInjectionComponent = NewGraphInstance->GetComponent<FGraphInstanceInjectionComponent>();
+			FGraphInstanceInjectionComponent& HostInjectionComponent = NewGraphInstance->GetOrAddComponent<FGraphInstanceInjectionComponent>();
 			WorkData.InjectedGraphReference = HostInjectionComponent.GetInjectionInfo().GetDefaultInjectionSite();
 		}
 	}
diff --git a/Engine/Plugins/Experimental/UAF/UAFAnimGraph/Source/UAFAnimGraph/Private/Injection/AnimNextModuleInjectionComponent.cpp b/Engine/Plugins/Experimental/UAF/UAFAnimGraph/Source/UAFAnimGraph/Private/Injection/AnimNextModuleInjectionComponent.cpp
--- a/Engine/Plugins/Experimental/UAF/UAFAnimGraph/Source/UAFAnimGraph/Private/Injection/AnimNextModuleInjectionComponent.cpp
+++ b/Engine/Plugins/Experimental/UAF/UAFAnimGraph/Source/UAFAnimGraph/Private/Injection/AnimNextModuleInjectionComponent.cpp
@@ -169,7 +169,7 @@ void FAnimNextModuleInjectionComponent::OnReapplyInjection(const UE::UAF::FModul
 {
 	using namespace UE::UAF;
 	
-	FAnimNextModuleInjectionComponent& Component = InContext.ModuleInstance.GetComponent<FAnimNextModuleInjectionComponent>();
+	FAnimNextModuleInjectionComponent& Component = InContext.ModuleInstance.GetOrAddComponent<FAnimNextModuleInjectionComponent>();
 
 	for(const TPair<FAnimNextVariableReference, FInjectionRecord>& RequestPair : Component.CurrentRequests)
 	{
diff --git a/Engine/Plugins/Experimental/UAF/UAFAnimGraph/Source/UAFAnimGraph/Private/Injection/InjectionSiteTrait.cpp b/Engine/Plugins/Experimental/UAF/UAFAnimGraph/Source/UAFAnimGraph/Private/Injection/InjectionSiteTrait.cpp
--- a/Engine/Plugins/Experimental/UAF/UAFAnimGraph/Source/UAFAnimGraph/Private/Injection/InjectionSiteTrait.cpp
+++ b/Engine/Plugins/Experimental/UAF/UAFAnimGraph/Source/UAFAnimGraph/Private/Injection/InjectionSiteTrait.cpp
@@ -63,7 +63,7 @@ namespace UE::UAF
 
 		if(FAnimNextModuleInstance* ModuleInstance = Context.GetRootGraphInstance().GetModuleInstance())
 		{
-			(void)ModuleInstance->GetComponent<FAnimNextModuleInjectionComponent>();
+			(void)ModuleInstance->GetOrAddComponent<FAnimNextModuleInjectionComponent>();
 		}
 	}
 
diff --git a/Engine/Plugins/Experimental/UAF/UAFAnimGraph/Source/UAFAnimGraph/Private/Injection/MontageTrait.cpp b/Engine/Plugins/Experimental/UAF/UAFAnimGraph/Source/UAFAnimGraph/Private/Injection/MontageTrait.cpp
--- a/Engine/Plugins/Experimental/UAF/UAFAnimGraph/Source/UAFAnimGraph/Private/Injection/MontageTrait.cpp
+++ b/Engine/Plugins/Experimental/UAF/UAFAnimGraph/Source/UAFAnimGraph/Private/Injection/MontageTrait.cpp
@@ -48,7 +48,7 @@ namespace UE::UAF
 		if(FAnimNextModuleInstance* ModuleInstance = Context.GetRootGraphInstance().GetModuleInstance())
 		{
 			// Cache the montage component and create it lazily if it hasn't been created yet. 
-			MontageComponent = &ModuleInstance->GetComponent<FUAFMontageComponent>();
+			MontageComponent = &ModuleInstance->GetOrAddComponent<FUAFMontageComponent>();
 		}
 
 		const FSharedData* SharedData = Binding.GetSharedData<FSharedData>();
@@ -341,7 +341,7 @@ namespace UE::UAF
 				if (FAnimNextModuleInstance* ModuleInstance = Context.GetRootGraphInstance().GetModuleInstance())
 				{
 					// Ensure we have a handler component on the module
-					(void)ModuleInstance->GetComponent<FUAFNotifyDispatcherComponent>();
+					(void)ModuleInstance->GetOrAddComponent<FUAFNotifyDispatcherComponent>();
 				}
 
 				TSharedPtr<FNotifyDispatchEvent> NotifyDispatchEvent = MakeTraitEvent<FNotifyDispatchEvent>();
@@ -374,7 +374,7 @@ namespace UE::UAF
 
 			// Register with synchronization component
 			const FSyncGroupParameters GroupParameters = InstanceData->SyncGroupParams;
-			FSyncGroupGraphInstanceComponent& SyncComponent = Context.GetComponent<FSyncGroupGraphInstanceComponent>();
+			FSyncGroupGraphInstanceComponent& SyncComponent = Context.GetOrAddComponent<FSyncGroupGraphInstanceComponent>();
 			SyncComponent.RegisterWithGroup(GroupParameters, Binding.GetTraitPtr(), TraitState);
 		}
 	}
diff --git a/Engine/Plugins/Experimental/UAF/UAFAnimGraph/Source/UAFAnimGraph/Private/TraitCore/ExecutionContext.cpp b/Engine/Plugins/Experimental/UAF/UAFAnimGraph/Source/UAFAnimGraph/Private/TraitCore/ExecutionContext.cpp
--- a/Engine/Plugins/Experimental/UAF/UAFAnimGraph/Source/UAFAnimGraph/Private/TraitCore/ExecutionContext.cpp
+++ b/Engine/Plugins/Experimental/UAF/UAFAnimGraph/Source/UAFAnimGraph/Private/TraitCore/ExecutionContext.cpp
@@ -610,7 +610,7 @@ namespace UE::UAF
 		{
 			if (FAnimNextModuleInstance* SystemInstance = RootGraphInstance->GetModuleInstance())
 			{
-				return SystemInstance->GetComponent<FUAFSystemOutputComponent>().GetSkeletalMeshComponent();
+				return SystemInstance->GetOrAddComponent<FUAFSystemOutputComponent>().GetSkeletalMeshComponent();
 			}
 		}
 		
diff --git a/Engine/Plugins/Experimental/UAF/UAFAnimGraph/Source/UAFAnimGraph/Private/TraitInterfaces/IEvaluate.cpp b/Engine/Plugins/Experimental/UAF/UAFAnimGraph/Source/UAFAnimGraph/Private/TraitInterfaces/IEvaluate.cpp
--- a/Engine/Plugins/Experimental/UAF/UAFAnimGraph/Source/UAFAnimGraph/Private/TraitInterfaces/IEvaluate.cpp
+++ b/Engine/Plugins/Experimental/UAF/UAFAnimGraph/Source/UAFAnimGraph/Private/TraitInterfaces/IEvaluate.cpp
@@ -68,7 +68,7 @@ namespace UE::UAF
 		// Pull the component from the system's output component
 		if (FAnimNextModuleInstance* SystemInstance = GraphInstance.GetModuleInstance())
 		{
-			return SystemInstance->GetComponent<FUAFSystemOutputComponent>().GetSkeletalMeshComponent();
+			return SystemInstance->GetOrAddComponent<FUAFSystemOutputComponent>().GetSkeletalMeshComponent();
 		}
 		return TWeakObjectPtr<const USkeletalMeshComponent>();
 	}
diff --git a/Engine/Plugins/Experimental/UAF/UAFAnimGraph/Source/UAFAnimGraph/Private/TraitInterfaces/IUpdate.cpp b/Engine/Plugins/Experimental/UAF/UAFAnimGraph/Source/UAFAnimGraph/Private/TraitInterfaces/IUpdate.cpp
--- a/Engine/Plugins/Experimental/UAF/UAFAnimGraph/Source/UAFAnimGraph/Private/TraitInterfaces/IUpdate.cpp
+++ b/Engine/Plugins/Experimental/UAF/UAFAnimGraph/Source/UAFAnimGraph/Private/TraitInterfaces/IUpdate.cpp
@@ -243,7 +243,7 @@ namespace UE::UAF
 		// Pull the component from the system's output component
 		if (FAnimNextModuleInstance* SystemInstance = GraphInstance.GetModuleInstance())
 		{
-			return SystemInstance->GetComponent<FUAFSystemOutputComponent>().GetSkeletalMeshComponent();
+			return SystemInstance->GetOrAddComponent<FUAFSystemOutputComponent>().GetSkeletalMeshComponent();
 		}
 		return TWeakObjectPtr<const USkeletalMeshComponent>();
 	}
diff --git a/Engine/Plugins/Experimental/UAF/UAFAnimGraph/Source/UAFAnimGraph/Private/Traits/NotifyDispatcher.cpp b/Engine/Plugins/Experimental/UAF/UAFAnimGraph/Source/UAFAnimGraph/Private/Traits/NotifyDispatcher.cpp
--- a/Engine/Plugins/Experimental/UAF/UAFAnimGraph/Source/UAFAnimGraph/Private/Traits/NotifyDispatcher.cpp
+++ b/Engine/Plugins/Experimental/UAF/UAFAnimGraph/Source/UAFAnimGraph/Private/Traits/NotifyDispatcher.cpp
@@ -67,7 +67,7 @@ namespace UE::UAF
 				// Ensure we have a handler component on the module
 				if(FAnimNextModuleInstance* ModuleInstance = Context.GetRootGraphInstance().GetModuleInstance())
 				{
-					(void)ModuleInstance->GetComponent<FUAFNotifyDispatcherComponent>();
+					(void)ModuleInstance->GetOrAddComponent<FUAFNotifyDispatcherComponent>();
 				}
 
 				auto NotifyDispatchEvent = MakeTraitEvent<FNotifyDispatchEvent>();
diff --git a/Engine/Plugins/Experimental/UAF/UAFAnimGraph/Source/UAFAnimGraph/Private/Traits/SynchronizeUsingGroups.cpp b/Engine/Plugins/Experimental/UAF/UAFAnimGraph/Source/UAFAnimGraph/Private/Traits/SynchronizeUsingGroups.cpp
--- a/Engine/Plugins/Experimental/UAF/UAFAnimGraph/Source/UAFAnimGraph/Private/Traits/SynchronizeUsingGroups.cpp
+++ b/Engine/Plugins/Experimental/UAF/UAFAnimGraph/Source/UAFAnimGraph/Private/Traits/SynchronizeUsingGroups.cpp
@@ -24,7 +24,7 @@ namespace UE::UAF
 
 	void FSynchronizeUsingGroupsTrait::FInstanceData::Construct(const FExecutionContext& Context, const FTraitBinding& Binding)
 	{
-		SyncGroupComponent = &Context.GetComponent<FSyncGroupGraphInstanceComponent>();
+		SyncGroupComponent = &Context.GetOrAddComponent<FSyncGroupGraphInstanceComponent>();
 
 		// If our timeline can't advance, we can't synchronize; assume we have no timeline
 		// We only check the base trait since it is the one that would own the timeline
diff --git a/Engine/Plugins/Experimental/UAF/UAFAnimGraph/Source/UAFAnimGraph/Public/TraitCore/ExecutionContext.h b/Engine/Plugins/Experimental/UAF/UAFAnimGraph/Source/UAFAnimGraph/Public/TraitCore/ExecutionContext.h
--- a/Engine/Plugins/Experimental/UAF/UAFAnimGraph/Source/UAFAnimGraph/Public/TraitCore/ExecutionContext.h
+++ b/Engine/Plugins/Experimental/UAF/UAFAnimGraph/Source/UAFAnimGraph/Public/TraitCore/ExecutionContext.h
@@ -170,7 +170,7 @@ namespace UE::UAF
 
 		// Returns a typed graph instance component, creating it lazily the first time it is queried
 		template<class ComponentType>
-		ComponentType& GetComponent() const;
+		ComponentType& GetOrAddComponent() const;
 
 		// Returns a typed graph instance component pointer if found or nullptr otherwise
 		template<class ComponentType>
@@ -329,11 +329,11 @@ namespace UE::UAF
 	}
 
 	template<class ComponentType>
-	ComponentType& FExecutionContext::GetComponent() const
+	ComponentType& FExecutionContext::GetOrAddComponent() const
 	{
 		static_assert(std::is_base_of<FUAFGraphInstanceComponent, ComponentType>::value, "ComponentType type must derive from FGraphInstanceComponent");
 
-		return RootGraphInstance->GetComponent<ComponentType>();
+		return RootGraphInstance->GetOrAddComponent<ComponentType>();
 	}
 
 	template<class ComponentType>
diff --git a/Engine/Plugins/Experimental/UAF/UAFAnimNode/Source/UAFAnimNode/Private/RigUnit_RunAnimNode.cpp b/Engine/Plugins/Experimental/UAF/UAFAnimNode/Source/UAFAnimNode/Private/RigUnit_RunAnimNode.cpp
--- a/Engine/Plugins/Experimental/UAF/UAFAnimNode/Source/UAFAnimNode/Private/RigUnit_RunAnimNode.cpp
+++ b/Engine/Plugins/Experimental/UAF/UAFAnimNode/Source/UAFAnimNode/Private/RigUnit_RunAnimNode.cpp
@@ -104,7 +104,7 @@ FRigUnit_RunAnimNode_Execute()
 
 		// Use AnimOps to compute our pose
 		{
-			FUAFSystemOutputComponent& OutputPoseComponent = ModuleInstance.GetComponent<FUAFSystemOutputComponent>();
+			FUAFSystemOutputComponent& OutputPoseComponent = ModuleInstance.GetOrAddComponent<FUAFSystemOutputComponent>();
 
 			// Note that the skeletal mesh and skeleton can be null
 			USkeletalMesh* SkeletalMesh = OutputPoseComponent.GetSkeletalMesh();
@@ -133,7 +133,7 @@ FRigUnit_RunAnimNode_Execute()
 			if (!Notifies.IsEmpty())
 			{
 				// Ensure we have a handler component on the module
-				(void)ModuleInstance.GetComponent<FUAFNotifyDispatcherComponent>();
+				(void)ModuleInstance.GetOrAddComponent<FUAFNotifyDispatcherComponent>();
 
 				auto NotifyDispatchEvent = MakeTraitEvent<FNotifyDispatchEvent>();
 				NotifyDispatchEvent->Notifies = MoveTemp(Notifies);
diff --git a/Engine/Plugins/Experimental/UAF/UAFLayering/Source/UAFLayering/Private/Traits/MontageLayerDataTrait.cpp b/Engine/Plugins/Experimental/UAF/UAFLayering/Source/UAFLayering/Private/Traits/MontageLayerDataTrait.cpp
--- a/Engine/Plugins/Experimental/UAF/UAFLayering/Source/UAFLayering/Private/Traits/MontageLayerDataTrait.cpp
+++ b/Engine/Plugins/Experimental/UAF/UAFLayering/Source/UAFLayering/Private/Traits/MontageLayerDataTrait.cpp
@@ -29,7 +29,7 @@ namespace UE::UAF
 			if (FAnimNextModuleInstance* ModuleInstance = Context.GetRootGraphInstance().GetModuleInstance())
 			{
 				FInstanceData* InstanceData = Binding.GetInstanceData<FInstanceData>();
-				FUAFMontageComponent& MontageComponent = ModuleInstance->GetComponent<FUAFMontageComponent>();
+				FUAFMontageComponent& MontageComponent = ModuleInstance->GetOrAddComponent<FUAFMontageComponent>();
 				const bool bSlotActive = MontageComponent.IsSlotActive(SharedData->SlotName);
 				
 				// Enable or disable the layer based on if a montage is currently playing in the relevant slot 
diff --git a/Engine/Plugins/Experimental/UAF/UAFStateTree/Source/UAFStateTree/Private/AnimNode/UAFStateTreeNode.cpp b/Engine/Plugins/Experimental/UAF/UAFStateTree/Source/UAFStateTree/Private/AnimNode/UAFStateTreeNode.cpp
--- a/Engine/Plugins/Experimental/UAF/UAFStateTree/Source/UAFStateTree/Private/AnimNode/UAFStateTreeNode.cpp
+++ b/Engine/Plugins/Experimental/UAF/UAFStateTree/Source/UAFStateTree/Private/AnimNode/UAFStateTreeNode.cpp
@@ -55,7 +55,7 @@ FUAFStateTreeNode::FUAFStateTreeNode(FUAFAnimGraphUpdateContext& Context, TObjec
 		StateTreeExternalParameters.Reset();
 	
 		TArray<int32> StartPropertyOffsets;
-		FUAFRigVMComponent& RigVMComponent = Context.GetVariablesOwner()->GetComponent<FUAFRigVMComponent>();
+		FUAFRigVMComponent& RigVMComponent = Context.GetVariablesOwner()->GetOrAddComponent<FUAFRigVMComponent>();
 		const FRigVMExtendedExecuteContext& ExtendedExecuteContext = RigVMComponent.GetExtendedExecuteContext();
 
 		// Variable Binding:  It might be possible to get this block of code working in a way where it binds the state tree variables to their values in the System/Module, but as-is the layouts mismatch and it fails.
diff --git a/Engine/Plugins/Experimental/UAF/UAFStateTree/Source/UAFStateTree/Private/AnimStateTreeTrait.cpp b/Engine/Plugins/Experimental/UAF/UAFStateTree/Source/UAFStateTree/Private/AnimStateTreeTrait.cpp
--- a/Engine/Plugins/Experimental/UAF/UAFStateTree/Source/UAFStateTree/Private/AnimStateTreeTrait.cpp
+++ b/Engine/Plugins/Experimental/UAF/UAFStateTree/Source/UAFStateTree/Private/AnimStateTreeTrait.cpp
@@ -119,7 +119,7 @@ namespace UE::UAF
 			InstanceData->StateTreeExternalParameters.Reset();
 			
 			TArray<int32> StartPropertyOffsets;
-			FUAFRigVMComponent& RigVMComponent = OwnerGraphInstance.GetComponent<FUAFRigVMComponent>();
+			FUAFRigVMComponent& RigVMComponent = OwnerGraphInstance.GetOrAddComponent<FUAFRigVMComponent>();
 			const FRigVMExtendedExecuteContext& ExtendedExecuteContext = RigVMComponent.GetExtendedExecuteContext();
 
 			const int32 NumVariables = StateTreeParameters.GetNumPropertiesInBag();
diff --git a/Engine/Plugins/Experimental/UAF/UAFStateTree/Source/UAFStateTree/Private/Conditions/AnimNextStateTreeRigVMConditionBase.cpp b/Engine/Plugins/Experimental/UAF/UAFStateTree/Source/UAFStateTree/Private/Conditions/AnimNextStateTreeRigVMConditionBase.cpp
--- a/Engine/Plugins/Experimental/UAF/UAFStateTree/Source/UAFStateTree/Private/Conditions/AnimNextStateTreeRigVMConditionBase.cpp
+++ b/Engine/Plugins/Experimental/UAF/UAFStateTree/Source/UAFStateTree/Private/Conditions/AnimNextStateTreeRigVMConditionBase.cpp
@@ -63,7 +63,7 @@ bool FAnimNextStateTreeRigVMConditionBase::TestCondition(FStateTreeExecutionCont
 				return false;
 			}
 
-			FUAFRigVMComponent& RigVMComponent = GraphInstance->GetComponent<FUAFRigVMComponent>();
+			FUAFRigVMComponent& RigVMComponent = GraphInstance->GetOrAddComponent<FUAFRigVMComponent>();
 			FRigVMExtendedExecuteContext& ExtendedExecuteContext = RigVMComponent.GetExtendedExecuteContext();
 			FAnimNextExecuteContext& AnimNextContext = ExtendedExecuteContext.GetPublicDataSafe<FAnimNextExecuteContext>();
 
diff --git a/Engine/Plugins/Experimental/UAF/UAFStateTree/Source/UAFStateTree/Private/Tasks/AnimNextStateTreeRigVMTaskBase.cpp b/Engine/Plugins/Experimental/UAF/UAFStateTree/Source/UAFStateTree/Private/Tasks/AnimNextStateTreeRigVMTaskBase.cpp
--- a/Engine/Plugins/Experimental/UAF/UAFStateTree/Source/UAFStateTree/Private/Tasks/AnimNextStateTreeRigVMTaskBase.cpp
+++ b/Engine/Plugins/Experimental/UAF/UAFStateTree/Source/UAFStateTree/Private/Tasks/AnimNextStateTreeRigVMTaskBase.cpp
@@ -94,7 +94,7 @@ EStateTreeRunStatus FAnimNextStateTreeRigVMTaskBase::CallFunction(FStateTreeExec
 				return EStateTreeRunStatus::Failed;
 			}
 
-			FUAFRigVMComponent& RigVMComponent = GraphInstance->GetComponent<FUAFRigVMComponent>();
+			FUAFRigVMComponent& RigVMComponent = GraphInstance->GetOrAddComponent<FUAFRigVMComponent>();
 			FRigVMExtendedExecuteContext& ExtendedExecuteContext = RigVMComponent.GetExtendedExecuteContext();
 			FAnimNextExecuteContext& AnimNextContext = ExtendedExecuteContext.GetPublicDataSafe<FAnimNextExecuteContext>();
 
diff --git a/Engine/Plugins/Experimental/UAF/UAFTestSuites/Source/UAFAnimGraphTestSuite/Private/AnimNextAnimGraphTraitGraphTest.cpp b/Engine/Plugins/Experimental/UAF/UAFTestSuites/Source/UAFAnimGraphTestSuite/Private/AnimNextAnimGraphTraitGraphTest.cpp
--- a/Engine/Plugins/Experimental/UAF/UAFTestSuites/Source/UAFAnimGraphTestSuite/Private/AnimNextAnimGraphTraitGraphTest.cpp
+++ b/Engine/Plugins/Experimental/UAF/UAFTestSuites/Source/UAFAnimGraphTestSuite/Private/AnimNextAnimGraphTraitGraphTest.cpp
@@ -57,7 +57,7 @@ namespace UE::UAF
 			const FSharedData* SharedData = Binding.GetSharedData<FSharedData>();
 			FInstanceData* InstanceData = Binding.GetInstanceData<FInstanceData>();
 
-			FRigVMExecuteContext& ExecuteContext = Context.GetRootGraphInstance().GetComponent<FUAFRigVMComponent>().GetExtendedExecuteContext().GetPublicData();
+			FRigVMExecuteContext& ExecuteContext = Context.GetRootGraphInstance().GetOrAddComponent<FUAFRigVMComponent>().GetExtendedExecuteContext().GetPublicData();
 			ExecuteContext.Logf(EMessageSeverity::Info, TEXT("UpdateCount == %d"), ++InstanceData->UpdateCount);
 			ExecuteContext.Logf(EMessageSeverity::Info, TEXT("SomeInt32 == %d"), SharedData->SomeInt32);
 			ExecuteContext.Logf(EMessageSeverity::Info, TEXT("SomeFloat == %.02f"), SharedData->SomeFloat);
@@ -75,7 +75,7 @@ namespace UE::UAF
 			IEvaluate::PostEvaluate(Context, Binding);
 
 			FInstanceData* InstanceData = Binding.GetInstanceData<FInstanceData>();
-			FRigVMExecuteContext& ExecuteContext = Context.GetRootGraphInstance().GetComponent<FUAFRigVMComponent>().GetExtendedExecuteContext().GetPublicData();
+			FRigVMExecuteContext& ExecuteContext = Context.GetRootGraphInstance().GetOrAddComponent<FUAFRigVMComponent>().GetExtendedExecuteContext().GetPublicData();
 			ExecuteContext.Logf(EMessageSeverity::Info, TEXT("EvaluateCount == %d"), ++InstanceData->EvaluateCount);
 		}
 	};
@@ -493,7 +493,7 @@ bool FAnimationAnimNextRuntimeTest_GraphExecute::RunTest(const FString& InParame
 		{
 			Messages.Add(Message);
 		});
-		GraphInstance->GetComponent<FUAFRigVMComponent>().GetExtendedExecuteContext().SetRuntimeSettings(RuntimeSettings);
+		GraphInstance->GetOrAddComponent<FUAFRigVMComponent>().GetExtendedExecuteContext().SetRuntimeSettings(RuntimeSettings);
 
 		{
 			FUpdateGraphContext UpdateGraphContext(*GraphInstance.Get(), 1.0f / 30.0f);
@@ -641,7 +641,7 @@ bool FAnimationAnimNextRuntimeTest_GraphExecuteLatent::RunTest(const FString& In
 		{
 			Messages.Add(Message);
 		});
-		GraphInstance->GetComponent<FUAFRigVMComponent>().GetExtendedExecuteContext().SetRuntimeSettings(RuntimeSettings);
+		GraphInstance->GetOrAddComponent<FUAFRigVMComponent>().GetExtendedExecuteContext().SetRuntimeSettings(RuntimeSettings);
 		
 		{
 			FUpdateGraphContext UpdateGraphContext(*GraphInstance.Get(), 1.0f / 30.0f);
diff --git a/Engine/Plugins/Experimental/UAF/UAFTestSuites/Source/UAFAnimGraphTestSuite/Private/GCInstanceDataGCTest.cpp b/Engine/Plugins/Experimental/UAF/UAFTestSuites/Source/UAFAnimGraphTestSuite/Private/GCInstanceDataGCTest.cpp
--- a/Engine/Plugins/Experimental/UAF/UAFTestSuites/Source/UAFAnimGraphTestSuite/Private/GCInstanceDataGCTest.cpp
+++ b/Engine/Plugins/Experimental/UAF/UAFTestSuites/Source/UAFAnimGraphTestSuite/Private/GCInstanceDataGCTest.cpp
@@ -53,7 +53,7 @@ namespace UE::UAF
 			
 			const bool bIsValid = InstanceData->WeakAnimSequence.IsValid() && !InstanceData->WeakAnimSequence->HasAnyFlags(RF_BeginDestroyed | RF_FinishDestroyed);
 			
-			FRigVMExecuteContext& ExecuteContext = Context.GetRootGraphInstance().GetComponent<FUAFRigVMComponent>().GetExtendedExecuteContext().GetPublicData();
+			FRigVMExecuteContext& ExecuteContext = Context.GetRootGraphInstance().GetOrAddComponent<FUAFRigVMComponent>().GetExtendedExecuteContext().GetPublicData();
 			ExecuteContext.Logf(EMessageSeverity::Info, TEXT("%s"), bIsValid ? TEXT("Valid") : TEXT("Invalid"));
 
 			++InstanceData->UpdateCount;
@@ -178,7 +178,7 @@ bool FAnimationUAFRuntimeTest_InstanceDataGC::RunTest(const FString& InParameter
 		{
 			Messages.Add(Message);
 		});
-		GraphInstance->GetComponent<FUAFRigVMComponent>().GetExtendedExecuteContext().SetRuntimeSettings(RuntimeSettings);
+		GraphInstance->GetOrAddComponent<FUAFRigVMComponent>().GetExtendedExecuteContext().SetRuntimeSettings(RuntimeSettings);
 		
 		FUpdateGraphContext UpdateGraphContext(*GraphInstance.Get(), 1.0f / 30.0f);
 		// First update, the anim sequence has been added to root so should definitely be alive
diff --git a/Engine/Plugins/Experimental/UAF/UAFTestSuites/Source/UAFAnimGraphTestSuite/Private/GCInstanceDataMissingGCTest.cpp b/Engine/Plugins/Experimental/UAF/UAFTestSuites/Source/UAFAnimGraphTestSuite/Private/GCInstanceDataMissingGCTest.cpp
--- a/Engine/Plugins/Experimental/UAF/UAFTestSuites/Source/UAFAnimGraphTestSuite/Private/GCInstanceDataMissingGCTest.cpp
+++ b/Engine/Plugins/Experimental/UAF/UAFTestSuites/Source/UAFAnimGraphTestSuite/Private/GCInstanceDataMissingGCTest.cpp
@@ -50,7 +50,7 @@ namespace UE::UAF
 			
 			const bool bIsValid = InstanceData->WeakAnimSequence.IsValid() && !InstanceData->WeakAnimSequence->HasAnyFlags(RF_BeginDestroyed | RF_FinishDestroyed);
 			
-			FRigVMExecuteContext& ExecuteContext = Context.GetRootGraphInstance().GetComponent<FUAFRigVMComponent>().GetExtendedExecuteContext().GetPublicData();
+			FRigVMExecuteContext& ExecuteContext = Context.GetRootGraphInstance().GetOrAddComponent<FUAFRigVMComponent>().GetExtendedExecuteContext().GetPublicData();
 			ExecuteContext.Logf(EMessageSeverity::Info, TEXT("%s"), bIsValid ? TEXT("Valid") : TEXT("Invalid"));
 
 			++InstanceData->UpdateCount;
@@ -163,7 +163,7 @@ bool FAnimationUAFRuntimeTest_InstanceDataGCMissing::RunTest(const FString& InPa
 		{
 			Messages.Add(Message);
 		});
-		GraphInstance->GetComponent<FUAFRigVMComponent>().GetExtendedExecuteContext().SetRuntimeSettings(RuntimeSettings);
+		GraphInstance->GetOrAddComponent<FUAFRigVMComponent>().GetExtendedExecuteContext().SetRuntimeSettings(RuntimeSettings);
 		
 		FUpdateGraphContext UpdateGraphContext(*GraphInstance.Get(), 1.0f / 30.0f);
 		// First update, the anim sequence has been added to root so should definitely be alive
diff --git a/Engine/Plugins/Experimental/UAF/UAFTestSuites/Source/UAFAnimGraphTestSuite/Private/GCSharedDataGCTest.cpp b/Engine/Plugins/Experimental/UAF/UAFTestSuites/Source/UAFAnimGraphTestSuite/Private/GCSharedDataGCTest.cpp
--- a/Engine/Plugins/Experimental/UAF/UAFTestSuites/Source/UAFAnimGraphTestSuite/Private/GCSharedDataGCTest.cpp
+++ b/Engine/Plugins/Experimental/UAF/UAFTestSuites/Source/UAFAnimGraphTestSuite/Private/GCSharedDataGCTest.cpp
@@ -66,7 +66,7 @@ namespace UE::UAF
 			
 			const bool bIsValid = InstanceData->WeakAnimSequence.IsValid() && !InstanceData->WeakAnimSequence->HasAnyFlags(RF_BeginDestroyed | RF_FinishDestroyed);
 			
-			FRigVMExecuteContext& ExecuteContext = Context.GetRootGraphInstance().GetComponent<FUAFRigVMComponent>().GetExtendedExecuteContext().GetPublicData();
+			FRigVMExecuteContext& ExecuteContext = Context.GetRootGraphInstance().GetOrAddComponent<FUAFRigVMComponent>().GetExtendedExecuteContext().GetPublicData();
 			ExecuteContext.Logf(EMessageSeverity::Info, TEXT("%s"), bIsValid ? TEXT("Valid") : TEXT("Invalid"));
 
 			++InstanceData->UpdateCount;
@@ -188,7 +188,7 @@ bool FAnimationUAFRuntimeTest_SharedDataGC::RunTest(const FString& InParameters)
 		{
 			Messages.Add(Message);
 		});
-		GraphInstance->GetComponent<FUAFRigVMComponent>().GetExtendedExecuteContext().SetRuntimeSettings(RuntimeSettings);
+		GraphInstance->GetOrAddComponent<FUAFRigVMComponent>().GetExtendedExecuteContext().SetRuntimeSettings(RuntimeSettings);
 		
 		FUpdateGraphContext UpdateGraphContext(*GraphInstance.Get(), 1.0f / 30.0f);
 		
diff --git a/Engine/Plugins/Experimental/UAF/UAFTestSuites/Source/UAFCQTestSuite/Private/BindableValueFunctionTests.cpp b/Engine/Plugins/Experimental/UAF/UAFTestSuites/Source/UAFCQTestSuite/Private/BindableValueFunctionTests.cpp
--- a/Engine/Plugins/Experimental/UAF/UAFTestSuites/Source/UAFCQTestSuite/Private/BindableValueFunctionTests.cpp
+++ b/Engine/Plugins/Experimental/UAF/UAFTestSuites/Source/UAFCQTestSuite/Private/BindableValueFunctionTests.cpp
@@ -39,7 +39,7 @@ struct FUAFBindingTestFixture
 		// Ensure the RigVM component exists — it may not be in the asset's default
 		// components, so GetComponent lazily creates it and calls OnBindToInstance
 		// which sets up the VM, execution context, and external variable runtime data.
-		[[maybe_unused]] FUAFRigVMComponent& RigVMComponent = InInstance.GetComponent<FUAFRigVMComponent>();
+		[[maybe_unused]] FUAFRigVMComponent& RigVMComponent = InInstance.GetOrAddComponent<FUAFRigVMComponent>();
 	}
 };
 
diff --git a/Engine/Plugins/Experimental/UAF/UAFTestSuites/Source/UAFTestSuite/Private/AnimNextVariablesTest.cpp b/Engine/Plugins/Experimental/UAF/UAFTestSuites/Source/UAFTestSuite/Private/AnimNextVariablesTest.cpp
--- a/Engine/Plugins/Experimental/UAF/UAFTestSuites/Source/UAFTestSuite/Private/AnimNextVariablesTest.cpp
+++ b/Engine/Plugins/Experimental/UAF/UAFTestSuites/Source/UAFTestSuite/Private/AnimNextVariablesTest.cpp
@@ -440,7 +440,7 @@ bool FVariables_UOLBindings::RunTest(const FString& InParameters)
 
 	FAnimNextModuleInstance Instance(Asset, NewObject<UUAFComponent>(), EAnimNextModuleInitMethod::None);
 	Instance.Initialize(FUAFSystemFactoryParams());
-	FUAFRigVMComponent& RigVMComponent = Instance.GetComponent<FUAFRigVMComponent>();
+	FUAFRigVMComponent& RigVMComponent = Instance.GetOrAddComponent<FUAFRigVMComponent>();
 
 	// Execute bindings to fetch data
 	RigVMComponent.CallEventByName(FAnimNextModuleContextData(Instance, FRigUnit_AnimNextExecuteBindings_GT::EventName, 0.0f));
