Fix FBindableXxx binding serialization issues when used with UAF traits

- Set/clear bindings via direct raw data modification + NotifyPostChange
(SetValueFromFormattedString fails for FStructOnScope-backed property handles)
- Always export __Binding=() marker even when no binding is set, so the RigVM
controller updates the parent pin DefaultValue (fixes clear binding on traits)
- Propagate ConstantValue child edits to parent handle via SetOnPropertyValueChanged
- Update serialization test for new __Binding=() export behavior
- Register FBindableValues with rig vm as compatible struct types

<description: restricted, no permission to view>

#was-backout 52597777
#Synced-CL 52638266

diff --git a/Engine/Plugins/Experimental/UAF/UAF/Source/UAF/Private/AnimNextModuleImpl.cpp b/Engine/Plugins/Experimental/UAF/UAF/Source/UAF/Private/AnimNextModuleImpl.cpp
--- a/Engine/Plugins/Experimental/UAF/UAF/Source/UAF/Private/AnimNextModuleImpl.cpp
+++ b/Engine/Plugins/Experimental/UAF/UAF/Source/UAF/Private/AnimNextModuleImpl.cpp
@@ -28,6 +28,7 @@
 #include "Variables/AnimNextSoftFunctionPtr.h"
 #include "RigVMCore/RigVMGraphFunctionDefinition.h"
 #include "HierarchyTable.h"
+#include "BindableValue/UAFBindableTypes.h"
 #include "CrashReporter/CrashReporterHandler.h"
 #include "Factory/SystemFactory.h"
 #include "Features/IModularFeatures.h"
@@ -92,6 +93,19 @@ namespace UE::UAF
 			FAnimNextGraphLODPose::StaticStruct(),
 			FAnimNextGraphReferencePose::StaticStruct(),
 			FUAFValueBundle::StaticStruct(),
+			FBindableBool::StaticStruct(),
+			FBindableFloat::StaticStruct(),
+			FBindableDouble::StaticStruct(),
+			FBindableInt32::StaticStruct(),
+			FBindableInt64::StaticStruct(),
+			FBindableByte::StaticStruct(),
+			FBindableName::StaticStruct(),
+			FBindableVector::StaticStruct(),
+			FBindableQuat::StaticStruct(),
+			FBindableTransform::StaticStruct(),
+			FBindableEnum::StaticStruct(),
+			FBindableStruct::StaticStruct(),
+			FBindableObject::StaticStruct(),
 		};
 
 		RigVMRegistry.RegisterStructTypes(AllowedStructTypes);
diff --git a/Engine/Plugins/Experimental/UAF/UAF/Source/UAF/Private/BindableValue/UAFBindableTypes.cpp b/Engine/Plugins/Experimental/UAF/UAF/Source/UAF/Private/BindableValue/UAFBindableTypes.cpp
--- a/Engine/Plugins/Experimental/UAF/UAF/Source/UAF/Private/BindableValue/UAFBindableTypes.cpp
+++ b/Engine/Plugins/Experimental/UAF/UAF/Source/UAF/Private/BindableValue/UAFBindableTypes.cpp
@@ -587,16 +587,34 @@ bool FBindableValueBase::Serialize(FArchive& Ar)
 
 static const TCHAR* BindingPropertyName = TEXT("__Binding");
 
+// IMPORTANT: __Binding must ALWAYS be exported, even when no binding is set.
+//
+// The RigVM controller's SetPinDefaultValue (RigVMController.cpp) compares the new text
+// against the pin's stored default value. For FBindable types with native import/export,
+// the pin has no sub-pins and the comparison happens on the full struct text. When the
+// editor customization modifies a binding (via direct raw data mutation + NotifyPostChange),
+// the notification chain exports the struct to text and calls SetPinDefaultValue. The
+// exported text must DIFFER from the pin's stored text so the controller detects the change:
+//     if(InPin->GetDefaultValue() != ClampedDefaultValue)  // RigVMController.cpp
+//
+// Without __Binding=(), clearing a binding produces the same text as the stored (pre-binding)
+// value, and the change is silently lost -- the pin default is never updated.
+//
+// DO NOT remove this always-export behavior without updating the editor binding customization
+// (UAFBindableValueCustomization.cpp) to use a different change-detection mechanism.
 void FBindableValueBase::ExportBindingText(FString& ValueStr, const FUAFPropertyBinding* InBinding, int32 PortFlags)
 {
-	if (!InBinding)
+	FString BindingText;
+	if (InBinding)
 	{
-		return;
+		FUAFPropertyBinding::StaticStruct()->ExportText(BindingText, InBinding, nullptr, nullptr, PortFlags, nullptr, true);
+	}
+	else
+	{
+		BindingText = TEXT("()");
 	}
 
-	FString BindingText;
-	FUAFPropertyBinding::StaticStruct()->ExportText(BindingText, InBinding, nullptr, nullptr, PortFlags, nullptr, true);
-	// Insert __Binding=(...) before the closing )
+	// Insert __Binding=(...) or __Binding=() before the closing )
 	if (ValueStr.Len() > 0 && ValueStr[ValueStr.Len() - 1] == TCHAR(')'))
 	{
 		// Only add comma separator if there are existing properties between the parens (more than just "()")
@@ -750,14 +768,23 @@ bool FBindableValueBase::ImportBindingText(const TCHAR*& Buffer, TUniquePtr<FUAF
 		if (*ValueStart == TCHAR('='))
 		{
 			ValueStart++;
-			OutBinding = MakeUnique<FUAFPropertyBinding>();
-			FUAFPropertyBinding::StaticStruct()->ImportText(ValueStart,
-				OutBinding.Get(),
-				nullptr,
-				PortFlags,
-				GLog,
-				TEXT("FUAFPropertyBinding"),
-				/*bAllowNativeOverride=*/false);
+			// __Binding=() is the empty marker emitted by ExportBindingText when no binding is set.
+			// Treat it as no-binding rather than importing a default-constructed FUAFPropertyBinding.
+			if (ValueStart[0] == TCHAR('(') && ValueStart[1] == TCHAR(')'))
+			{
+				OutBinding.Reset();
+			}
+			else
+			{
+				OutBinding = MakeUnique<FUAFPropertyBinding>();
+				FUAFPropertyBinding::StaticStruct()->ImportText(ValueStart,
+					OutBinding.Get(),
+					nullptr,
+					PortFlags,
+					GLog,
+					TEXT("FUAFPropertyBinding"),
+					/*bAllowNativeOverride=*/false);
+			}
 		}
 	}
 	else
diff --git a/Engine/Plugins/Experimental/UAF/UAF/Source/UAFEditor/Private/BindableValue/UAFBindableValueCustomization.cpp b/Engine/Plugins/Experimental/UAF/UAF/Source/UAFEditor/Private/BindableValue/UAFBindableValueCustomization.cpp
--- a/Engine/Plugins/Experimental/UAF/UAF/Source/UAFEditor/Private/BindableValue/UAFBindableValueCustomization.cpp
+++ b/Engine/Plugins/Experimental/UAF/UAF/Source/UAFEditor/Private/BindableValue/UAFBindableValueCustomization.cpp
@@ -376,43 +376,19 @@ void FBindableValueCustomization::CustomizeHeader(
 		return bResult;
 	};
 
-	// Clears the binding via copy + export + SetValueFromFormattedString.
-	// We must NOT modify raw data before SetValueFromFormattedString, because the property
-	// system compares the new text against the current in-memory state. If they match (because
-	// we already mutated raw data), no change is detected and the pin default is never updated.
-	auto ClearBinding = [Handle = PropertyHandle]()
+	// Clears the binding by modifying raw data directly and firing a property change notification.
+	// Direct modification is used instead of SetValueFromFormattedString because the latter
+	// fails silently for property handles backed by FStructOnScope (e.g. trait properties).
+	auto ClearBindingDirect = [](TSharedPtr<IPropertyHandle> Handle)
 	{
-		const FStructProperty* StructProp = CastField<FStructProperty>(Handle->GetProperty());
-		if (!StructProp)
-		{
-			return;
-		}
-
-		FString TextValue;
-		Handle->EnumerateConstRawData(
-			[&TextValue, StructProp](const void* RawData, const int32, const int32) -> bool
+		Handle->NotifyPreChange();
+		Handle->EnumerateRawData(
+			[](void* RawData, const int32, const int32) -> bool
 			{
-				UScriptStruct* Struct = StructProp->Struct;
-				const int32 StructSize = Struct->GetStructureSize();
-
-				TArray<uint8> CopyBuf;
-				CopyBuf.SetNumUninitialized(StructSize);
-				Struct->InitializeStruct(CopyBuf.GetData());
-				Struct->CopyScriptStruct(CopyBuf.GetData(), RawData);
-				static_cast<FBindableValueBase*>(static_cast<void*>(CopyBuf.GetData()))->ClearBinding();
-
-				TArray<uint8> DefaultBuf;
-				DefaultBuf.SetNumUninitialized(StructSize);
-				Struct->InitializeStruct(DefaultBuf.GetData());
-
-				Struct->ExportText(TextValue, CopyBuf.GetData(), DefaultBuf.GetData(), nullptr, 0, nullptr, true);
-
-				Struct->DestroyStruct(DefaultBuf.GetData());
-				Struct->DestroyStruct(CopyBuf.GetData());
-				return false;
+				static_cast<FBindableValueBase*>(RawData)->ClearBinding();
+				return true;
 			});
-
-		Handle->SetValueFromFormattedString(TextValue);
+		Handle->NotifyPostChange(EPropertyChangeType::ValueSet);
 	};
 
 	// -----------------------------------------------------------------------
@@ -477,62 +453,36 @@ void FBindableValueCustomization::CustomizeHeader(
 		PickerArgs.CompatibleTypes = { TargetParamType };
 	}
 
-	// Helper: sets a binding by exporting from a copy and calling SetValueFromFormattedString.
-	// We must NOT modify the property handle's raw data before SetValueFromFormattedString,
-	// because the property system compares the new text against the current in-memory state.
-	// If they match (because we already mutated raw data), no change is detected and the
-	// AnimNext editor's UpdatePinDefaultValue callback never fires.
-	auto SetBindingViaText = [](TSharedPtr<IPropertyHandle> Handle, const FUAFPropertyBinding& InBinding)
+	// Helper: sets a binding by modifying raw data directly and firing a property change notification.
+	// Direct modification is used instead of SetValueFromFormattedString because the latter
+	// fails silently for property handles backed by FStructOnScope (e.g. trait properties).
+	// The notification triggers SetOnPropertyValueChangedWithData callbacks which export the
+	// full struct value (including binding via custom ExportTextItem) to the RigVM pin.
+	auto SetBindingDirect = [](TSharedPtr<IPropertyHandle> Handle, const FUAFPropertyBinding& InBinding)
 	{
-		const FStructProperty* StructProp = CastField<FStructProperty>(Handle->GetProperty());
-		if (!StructProp)
-		{
-			return;
-		}
-
-		FString TextValue;
-		Handle->EnumerateConstRawData(
-			[&TextValue, &InBinding, StructProp](const void* RawData, const int32, const int32) -> bool
+		Handle->NotifyPreChange();
+		Handle->EnumerateRawData(
+			[&InBinding](void* RawData, const int32, const int32) -> bool
 			{
-				UScriptStruct* Struct = StructProp->Struct;
-				const int32 StructSize = Struct->GetStructureSize();
-
-				// Make a copy of the current value and set the binding on the copy
-				TArray<uint8> CopyBuf;
-				CopyBuf.SetNumUninitialized(StructSize);
-				Struct->InitializeStruct(CopyBuf.GetData());
-				Struct->CopyScriptStruct(CopyBuf.GetData(), RawData);
-				static_cast<FBindableValueBase*>(static_cast<void*>(CopyBuf.GetData()))->SetBinding(InBinding);
-
-				// Export text from the copy (property handle data is untouched)
-				TArray<uint8> DefaultBuf;
-				DefaultBuf.SetNumUninitialized(StructSize);
-				Struct->InitializeStruct(DefaultBuf.GetData());
-
-				Struct->ExportText(TextValue, CopyBuf.GetData(), DefaultBuf.GetData(), nullptr, 0, nullptr, /*bAllowNativeOverride=*/true);
-
-				Struct->DestroyStruct(DefaultBuf.GetData());
-				Struct->DestroyStruct(CopyBuf.GetData());
-				return false; // first instance only
+				static_cast<FBindableValueBase*>(RawData)->SetBinding(InBinding);
+				return true;
 			});
-
-		// Text differs from current state → property system detects the change → fires notification
-		Handle->SetValueFromFormattedString(TextValue);
+		Handle->NotifyPostChange(EPropertyChangeType::ValueSet);
 	};
 
 	// Write a Variable binding.
 	PickerArgs.OnVariablePicked = FOnVariablePicked::CreateLambda(
-		[this, SetBindingViaText](const FAnimNextSoftVariableReference& SoftRef, const FAnimNextParamType& /*Type*/)
+		[this, SetBindingDirect](const FAnimNextSoftVariableReference& SoftRef, const FAnimNextParamType& /*Type*/)
 		{
 			FUAFPropertyBinding NewBinding;
 			NewBinding.SourceType     = EUAFBindingSourceType::Variable;
 			NewBinding.SourceVariable = FAnimNextVariableReference(SoftRef);
-			SetBindingViaText(PropertyHandle, NewBinding);
+			SetBindingDirect(PropertyHandle, NewBinding);
 		});
 
 	// Write a SubProperty binding.
 	PickerArgs.OnSubPropertyVariablePicked = FOnSubPropertyVariablePicked::CreateLambda(
-		[this, SetBindingViaText](
+		[this, SetBindingDirect](
 			const FAnimNextSoftVariableReference& TopVarRef,
 			const TArray<FName>&                  SubPropertyPath,
 			const FAnimNextParamType& /*LeafType*/)
@@ -544,7 +494,7 @@ void FBindableValueCustomization::CustomizeHeader(
 			{
 				NewBinding.SubPropertyPath.AddPathSegment(Seg);
 			}
-			SetBindingViaText(PropertyHandle, NewBinding);
+			SetBindingDirect(PropertyHandle, NewBinding);
 
 			FSlateApplication::Get().DismissAllMenus();
 		});
@@ -709,6 +659,18 @@ void FBindableValueCustomization::CustomizeHeader(
 		}
 	}
 
+	// When the constant value changes through the widget, also notify the parent FBindableXxx
+	// handle. This ensures that external callbacks (e.g. trait pin default update) that are
+	// registered on the parent handle see the change and export the full struct value.
+	if (ConstantValueHandle.IsValid())
+	{
+		ConstantValueHandle->SetOnPropertyValueChanged(FSimpleDelegate::CreateLambda([this]()
+		{
+			PropertyHandle->NotifyPreChange();
+			PropertyHandle->NotifyPostChange(EPropertyChangeType::ValueSet);
+		}));
+	}
+
 	// Validate function binding once at widget creation time (not per-frame).
 	// Also resolve the current display name from GUID (handles function renames).
 	EFunctionBindingStatus CachedFunctionStatus = EFunctionBindingStatus::Valid;
@@ -898,7 +860,7 @@ void FBindableValueCustomization::CustomizeHeader(
 				]
 			]
 			.OnGetMenuContent_Lambda(
-				[this, PickerArgs, SetBindingViaText, TargetParamType]() -> TSharedRef<SWidget>
+				[this, PickerArgs, SetBindingDirect, TargetParamType]() -> TSharedRef<SWidget>
 			{
 				TSharedRef<SWidgetSwitcher> Switcher = SNew(SWidgetSwitcher);
 
@@ -923,7 +885,7 @@ void FBindableValueCustomization::CustomizeHeader(
 						return IsFunctionCompatibleWithBindableType(InHeader, TargetParamType);
 					})
 					.OnRigVMFunctionPicked_Lambda(
-						[Handle = PropertyHandle, SetBindingViaText, WeakAsset = CurrentRigVMAsset]
+						[Handle = PropertyHandle, SetBindingDirect, WeakAsset = CurrentRigVMAsset]
 						(const FRigVMGraphFunctionHeader& InHeader)
 					{
 						if (!InHeader.IsValid())
@@ -935,7 +897,7 @@ void FBindableValueCustomization::CustomizeHeader(
 						FUAFPropertyBinding NewBinding;
 						NewBinding.SourceType = EUAFBindingSourceType::Function;
 						NewBinding.SourceFunction = FAnimNextFunctionReference::FromHeader(InHeader, Asset);
-						SetBindingViaText(Handle, NewBinding);
+						SetBindingDirect(Handle, NewBinding);
 					})
 				];
 
@@ -981,9 +943,9 @@ void FBindableValueCustomization::CustomizeHeader(
 			{
 				return HasBinding() ? EVisibility::Visible : EVisibility::Collapsed;
 			})
-			.OnClicked_Lambda([ClearBinding]() -> FReply
+			.OnClicked_Lambda([Handle = PropertyHandle, ClearBindingDirect]() -> FReply
 			{
-				ClearBinding();
+				ClearBindingDirect(Handle);
 				return FReply::Handled();
 			})
 			[
@@ -997,7 +959,7 @@ void FBindableValueCustomization::CustomizeHeader(
 	// Right-click "Clear Binding" context menu action on the detail row.
 	InHeaderRow.AddCustomContextMenuAction(
 		FUIAction(
-			FExecuteAction::CreateLambda([ClearBinding]() { ClearBinding(); }),
+			FExecuteAction::CreateLambda([Handle = PropertyHandle, ClearBindingDirect]() { ClearBindingDirect(Handle); }),
 			FCanExecuteAction::CreateLambda([HasBinding]() { return HasBinding(); })),
 		LOCTEXT("ClearBindingAction",        "Clear Binding"),
 		LOCTEXT("ClearBindingActionTooltip", "Remove the variable binding from this property"),
diff --git a/Engine/Plugins/Experimental/UAF/UAF/Tests/UAF/Private/BindableValueSerializationTests.cpp b/Engine/Plugins/Experimental/UAF/UAF/Tests/UAF/Private/BindableValueSerializationTests.cpp
--- a/Engine/Plugins/Experimental/UAF/UAF/Tests/UAF/Private/BindableValueSerializationTests.cpp
+++ b/Engine/Plugins/Experimental/UAF/UAF/Tests/UAF/Private/BindableValueSerializationTests.cpp
@@ -98,46 +98,23 @@ BindableType EditorStyleSetBinding(const BindableType& Source, const FUAFPropert
 	return Dest;
 }
 
-/** Mimics the editor flow: export current state (with binding), clear binding, import back. */
+/** Mimics the editor flow: copy struct, clear binding on copy, export, import back.
+ *  This matches the real editor path: EnumerateRawData + ClearBinding + NotifyPostChange.
+ *  The exported text will contain __Binding=() (the always-emit empty marker), and ImportTextItem
+ *  treats that as no-binding. */
 template<typename BindableType>
 BindableType EditorStyleClearBinding(const BindableType& Source)
 {
-	// Export with binding, then re-import without it — simulates clearing via text round-trip
+	// Step 1: Copy source and clear binding (like EnumerateRawData + ClearBinding on a copy)
+	BindableType WithCleared = Source;
+	WithCleared.ClearBinding();
+
+	// Step 2: Export to text (like GetValueAsFormattedString -> ExportTextItem)
 	FString TextValue;
 	BindableType Default;
-	Source.ExportTextItem(TextValue, Default, nullptr, 0, nullptr);
-
-	// Strip the __Binding pseudo-property from the text
-	const int32 BindingIdx = TextValue.Find(TEXT("__Binding"));
-	if (BindingIdx != INDEX_NONE)
-	{
-		// Find the comma before __Binding (or opening paren if it's the only member)
-		int32 StartIdx = BindingIdx;
-		if (StartIdx > 0 && TextValue[StartIdx - 1] == TCHAR(','))
-		{
-			StartIdx--;
-		}
-		// Find the end of the binding value — match nested parens
-		int32 EndIdx = BindingIdx;
-		int32 Depth = 0;
-		bool bFoundEquals = false;
-		while (EndIdx < TextValue.Len())
-		{
-			const TCHAR Ch = TextValue[EndIdx];
-			if (!bFoundEquals)
-			{
-				if (Ch == TCHAR('=')) { bFoundEquals = true; }
-			}
-			else
-			{
-				if (Ch == TCHAR('(')) { Depth++; }
-				else if (Ch == TCHAR(')')) { if (--Depth <= 0) { EndIdx++; break; } }
-			}
-			EndIdx++;
-		}
-		TextValue.RemoveAt(StartIdx, EndIdx - StartIdx);
-	}
+	WithCleared.ExportTextItem(TextValue, Default, nullptr, 0, nullptr);
 
+	// Step 3: Import back (like SetValueFromFormattedString -> ImportText_Direct)
 	BindableType Dest;
 	const TCHAR* Buffer = *TextValue;
 	const bool bImported = Dest.ImportTextItem(Buffer, 0, nullptr, nullptr);
@@ -306,7 +283,9 @@ TEST_CASE("FBindableBool.TextSerialize.NoBinding", "[UAF][unit][MustPass]")
 	FString Text;
 	FBindableBool Default;
 	Src.ExportTextItem(Text, Default, nullptr, 0, nullptr);
-	CHECK_FALSE(Text.Contains(TEXT("__Binding")));
+	// __Binding=() must ALWAYS be emitted (even without a binding) so that the RigVM
+	// controller's text comparison detects binding state changes. See ExportBindingText.
+	CHECK(Text.Contains(TEXT("__Binding=()")));
 
 	FBindableBool Dst = TextRoundTrip(Src);
 	CHECK(Dst.GetConstantValue() == true);
@@ -760,7 +739,7 @@ TEST_CASE("FBindableVector.ImportTextItem.DoesNotPickUpSiblingBinding", "[UAF][u
 
 TEST_CASE("FBindableTransform.ImportTextItem.DoesNotPickUpSiblingBinding", "[UAF][unit][MustPass]")
 {
-	// FTransform has deeply nested parens — verify struct boundary detection handles it
+	// FTransform has deeply nested parens -- verify struct boundary detection handles it
 	const FString Buffer = TEXT("(ConstantValue=(Rotation=(X=0.000000,Y=0.000000,Z=0.000000,W=1.000000),Translation=(X=1.000000,Y=2.000000,Z=3.000000),Scale3D=(X=1.000000,Y=1.000000,Z=1.000000))),Sibling=(ConstantValue=0,__Binding=(SourceType=Variable,SourceVariable=(Name=\"FloatVal\")))");
 	const TCHAR* Ptr = *Buffer;
 
@@ -770,3 +749,54 @@ TEST_CASE("FBindableTransform.ImportTextItem.DoesNotPickUpSiblingBinding", "[UAF
 	CHECK(Dst.GetConstantValue().GetTranslation() == FVector(1.0, 2.0, 3.0));
 	CHECK_FALSE(Dst.HasBinding());
 }
+
+// ===========================================================================
+// Binding-clear change detection tests -- verifies that the RigVM controller's
+// SetPinDefaultValue text comparison (RigVMController.cpp) can detect when a
+// binding is set or cleared. This is the bug that the always-emit __Binding=()
+// behavior in ExportBindingText was introduced to fix.
+// ===========================================================================
+
+TEST_CASE("FBindableFloat.TextSerialize.BindingClearProducesDifferentText", "[UAF][unit][MustPass]")
+{
+	FBindableFloat Default;
+
+	// 1. Export with no binding
+	FBindableFloat Unbound(42.5f);
+	FString UnboundText;
+	Unbound.ExportTextItem(UnboundText, Default, nullptr, 0, nullptr);
+
+	// 2. Set a binding and export
+	FBindableFloat Bound(42.5f);
+	Bound.SetBinding(MakeTestBinding());
+	FString BoundText;
+	Bound.ExportTextItem(BoundText, Default, nullptr, 0, nullptr);
+
+	// 3. Clear the binding and export
+	FBindableFloat Cleared(42.5f);
+	Cleared.SetBinding(MakeTestBinding());
+	Cleared.ClearBinding();
+	FString ClearedText;
+	Cleared.ExportTextItem(ClearedText, Default, nullptr, 0, nullptr);
+
+	// Bound text differs from unbound (binding is visible in exported text)
+	CHECK(BoundText != UnboundText);
+	// Cleared text differs from bound (clearing is detectable by text comparison)
+	CHECK(ClearedText != BoundText);
+	// Cleared text contains the empty binding marker
+	CHECK(ClearedText.Contains(TEXT("__Binding=()")));
+	// Unbound also has the empty marker (always-emit behavior)
+	CHECK(UnboundText.Contains(TEXT("__Binding=()")));
+}
+
+TEST_CASE("FBindableFloat.TextSerialize.HashStability", "[UAF][unit][MustPass]")
+{
+	FBindableFloat Src(42.5f);
+	FBindableFloat Default;
+
+	FString Text1, Text2;
+	Src.ExportTextItem(Text1, Default, nullptr, 0, nullptr);
+	Src.ExportTextItem(Text2, Default, nullptr, 0, nullptr);
+
+	CHECK(GetTypeHash(Text1) == GetTypeHash(Text2));
+}
diff --git a/Engine/Plugins/Experimental/UAF/UAFTestSuites/Source/UAFAnimNodeTestData/Private/UAFTestBindableTraitData.cpp b/Engine/Plugins/Experimental/UAF/UAFTestSuites/Source/UAFAnimNodeTestData/Private/UAFTestBindableTraitData.cpp
new file mode 100644
--- /dev/null
+++ b/Engine/Plugins/Experimental/UAF/UAFTestSuites/Source/UAFAnimNodeTestData/Private/UAFTestBindableTraitData.cpp
@@ -0,0 +1,10 @@
+// Copyright Epic Games, Inc. All Rights Reserved.
+
+#include "UAFTestBindableTraitData.h"
+#include "TraitCore/TraitBinding.h"
+
+namespace UE::UAF
+{
+	AUTO_REGISTER_ANIM_TRAIT(FTestBindableTrait)
+	GENERATE_ANIM_TRAIT_IMPLEMENTATION(FTestBindableTrait, NULL_ANIM_TRAIT_INTERFACE_ENUMERATOR, NULL_ANIM_TRAIT_INTERFACE_ENUMERATOR, NULL_ANIM_TRAIT_EVENT_ENUMERATOR)
+}
diff --git a/Engine/Plugins/Experimental/UAF/UAFTestSuites/Source/UAFAnimNodeTestData/Public/UAFTestBindableTraitData.h b/Engine/Plugins/Experimental/UAF/UAFTestSuites/Source/UAFAnimNodeTestData/Public/UAFTestBindableTraitData.h
new file mode 100644
--- /dev/null
+++ b/Engine/Plugins/Experimental/UAF/UAFTestSuites/Source/UAFAnimNodeTestData/Public/UAFTestBindableTraitData.h
@@ -0,0 +1,47 @@
+// Copyright Epic Games, Inc. All Rights Reserved.
+
+#pragma once
+
+#include "CoreMinimal.h"
+#include "TraitCore/Trait.h"
+#include "TraitCore/TraitBinding.h"
+#include "TraitCore/TraitSharedData.h"
+#include "BindableValue/UAFBindableTypes.h"
+
+#include "UAFTestBindableTraitData.generated.h"
+
+/**
+ * Minimal test-only trait shared data with a single FBindableBool property.
+ * Used by CQTests to exercise SetPinDefaultValue round-trip on FBindable
+ * types within trait pins, verifying that the RigVM controller detects
+ * value changes through text comparison.
+ */
+USTRUCT()
+struct FUAFTestBindableTraitSharedData : public FAnimNextTraitSharedData
+{
+	GENERATED_BODY()
+
+	UPROPERTY(EditAnywhere, Category = "Test")
+	FBindableBool bTestBool = false;
+
+#define TRAIT_LATENT_PROPERTIES_ENUMERATOR(GeneratorMacro) \
+	GeneratorMacro(bTestBool) \
+
+	GENERATE_TRAIT_LATENT_PROPERTIES(FUAFTestBindableTraitSharedData, TRAIT_LATENT_PROPERTIES_ENUMERATOR)
+#undef TRAIT_LATENT_PROPERTIES_ENUMERATOR
+};
+
+namespace UE::UAF
+{
+	/**
+	 * Minimal test-only trait that uses FUAFTestBindableTraitSharedData.
+	 * Required so that FRigDecorator_AnimNextCppDecorator can resolve the
+	 * shared data struct to a registered trait during AddTrait.
+	 */
+	struct FTestBindableTrait : FBaseTrait
+	{
+		DECLARE_ANIM_TRAIT(FTestBindableTrait, FBaseTrait)
+
+		using FSharedData = FUAFTestBindableTraitSharedData;
+	};
+}
diff --git a/Engine/Plugins/Experimental/UAF/UAFTestSuites/Source/UAFCQTestSuite/Private/BindableValueTraitPinTests.cpp b/Engine/Plugins/Experimental/UAF/UAFTestSuites/Source/UAFCQTestSuite/Private/BindableValueTraitPinTests.cpp
new file mode 100644
--- /dev/null
+++ b/Engine/Plugins/Experimental/UAF/UAFTestSuites/Source/UAFCQTestSuite/Private/BindableValueTraitPinTests.cpp
@@ -0,0 +1,306 @@
+// Copyright Epic Games, Inc. All Rights Reserved.
+
+#if WITH_DEV_AUTOMATION_TESTS
+#if WITH_EDITOR
+
+#include "UAFTestsUtilities.h"
+#include "CQTest.h"
+#include "Editor.h"
+#include "Subsystems/AssetEditorSubsystem.h"
+#include "Graph/AnimNextAnimationGraph.h"
+#include "Graph/AnimNextAnimationGraphFactory.h"
+#include "Graph/RigUnit_AnimNextTraitStack.h"
+#include "RigVMModel/RigVMController.h"
+#include "RigVMModel/RigVMPin.h"
+#include "UAFTestBlueprintLibrary.h"
+#include "UAFTestBindableTraitData.h"
+#include "BindableValue/UAFPropertyBinding.h"
+
+// Exercises the FBindableBool public API (SetBinding, ClearBinding, Value mutation)
+// and verifies the exported text changes are properly detected by the RigVM controller.
+//
+// The trait system decomposes FBindableBool into a ConstantValue sub-pin via the
+// latent property system. This test verifies both:
+//  - Direct constant value manipulation through the sub-pin
+//  - The FBindableBool public API produces correctly differentiated export text,
+//    which is the mechanism the RigVM controller uses for change detection
+TEST_CLASS(FBindableValueTraitPinTests, "Animation.UAF.Functional")
+{
+	UUAFAnimGraph* AnimNextAnimationGraph = nullptr;
+	UAssetEditorSubsystem* AssetEditorSubsystem = nullptr;
+	URigVMController* Controller = nullptr;
+	URigVMGraph* ControllerGraph = nullptr;
+	const FString TraitStackNodeName = TEXT("AnimNextTraitStack");
+	const FString TestTraitName = TEXT("FTestBindableTrait");
+	const FString BoolPinName = GET_MEMBER_NAME_STRING_CHECKED(FUAFTestBindableTraitSharedData, bTestBool);
+	// ConstantValue is private on FBindableBool -- validated via FindPin at runtime
+	const FString ConstantValueName = TEXT("ConstantValue");
+
+	BEFORE_EACH()
+	{
+		AssetEditorSubsystem = GEditor->GetEditorSubsystem<UAssetEditorSubsystem>();
+		ASSERT_THAT(IsNotNull(AssetEditorSubsystem));
+	}
+
+	// Verifies that the FBindableBool public API (SetBinding, ClearBinding)
+	// produces export text where each distinct binding state maps to a unique
+	// string. This is the core property the RigVM controller depends on: it
+	// compares exported text to detect changes (RigVMController.cpp
+	// SetPinDefaultValue). Without the always-emit __Binding=() marker,
+	// clearing a binding would produce text identical to the original unbound
+	// state, and the controller would silently discard the change.
+	//
+	// Uses ConstantValue=true (non-default) because ExportTextItem's internal
+	// self-comparison technique (passes this as both Value and Defaults) relies
+	// on sub-property values being non-default to survive the ExportText_Direct
+	// Identical check. The LLT BindableValueSerializationTests follow the same
+	// pattern.
+	TEST_METHOD(FBindableBool_PublicAPI_ExportTextDiffersPerState)
+	{
+		TestCommandBuilder
+			.Do(TEXT("Verify export text differs for each FBindableBool state"), [&]()
+			{
+				const UScriptStruct* BoolStruct = FBindableBool::StaticStruct();
+				ASSERT_THAT(IsNotNull(BoolStruct));
+
+				FUAFPropertyBinding TestBinding;
+				TestBinding.SourceType = EUAFBindingSourceType::Variable;
+
+				// --- State 1: unbound (ConstantValue=true, no binding) ---
+				FBindableBool Unbound(true);
+
+				// --- State 2: bound (ConstantValue=true, binding set) ---
+				FBindableBool Bound(true);
+				Bound.SetBinding(TestBinding);
+				ASSERT_THAT(IsTrue(Bound.HasBinding()));
+
+				// --- State 3: cleared (ConstantValue=true, binding set then cleared) ---
+				FBindableBool Cleared(true);
+				Cleared.SetBinding(TestBinding);
+				Cleared.ClearBinding();
+				ASSERT_THAT(IsFalse(Cleared.HasBinding()));
+
+				// Export each state to text via UScriptStruct::ExportText, which
+				// dispatches to FBindableBool::ExportTextItem (native override).
+				// This is the same path the RigVM uses to serialize pin values.
+				FString UnboundText, BoundText, ClearedText;
+				BoolStruct->ExportText(UnboundText, &Unbound, nullptr, nullptr, 0, nullptr, true);
+				BoolStruct->ExportText(BoundText, &Bound, nullptr, nullptr, 0, nullptr, true);
+				BoolStruct->ExportText(ClearedText, &Cleared, nullptr, nullptr, 0, nullptr, true);
+
+				// The RigVM controller compares pin text via string equality:
+				//   if (InPin->GetDefaultValue() != ClampedDefaultValue)
+				// Each distinct binding state must produce unique text so the
+				// controller detects every transition.
+
+				// Setting a binding must produce different text from unbound
+				ASSERT_THAT(IsTrue(UnboundText != BoundText));
+
+				// Clearing a binding must produce different text from bound state
+				ASSERT_THAT(IsTrue(ClearedText != BoundText));
+
+				// Both unbound and cleared should contain __Binding=() -- the
+				// always-emit empty marker that ensures change detection works
+				ASSERT_THAT(IsTrue(UnboundText.Contains(TEXT("__Binding=()"))));
+				ASSERT_THAT(IsTrue(ClearedText.Contains(TEXT("__Binding=()"))));
+
+				// The bound text must contain binding data, not just the empty marker
+				ASSERT_THAT(IsTrue(BoundText.Contains(TEXT("__Binding="))));
+				ASSERT_THAT(IsFalse(BoundText.Contains(TEXT("__Binding=()"))));
+			});
+	}
+
+	// Verifies that FBindableBool on a trait pin round-trips constant value
+	// changes through the RigVM controller's SetPinDefaultValue /
+	// GetDefaultValue. With GENERATE_TRAIT_LATENT_PROPERTIES, the trait system
+	// decomposes FBindableBool into a ConstantValue sub-pin for the literal
+	// bool value.
+	TEST_METHOD(FBindableBool_RoundTripThroughController)
+	{
+		TestCommandBuilder
+			.Do(TEXT("Create AnimNext Animation Graph Asset"), [&]()
+			{
+				CreateAnimNextAnimationGraph();
+			})
+			.Then(TEXT("Create a Trait Stack node"), [&]()
+			{
+				CreateTraitStackNode();
+			})
+			.Then(TEXT("Add test trait with FBindableBool"), [&]()
+			{
+				const FString SharedDataPath = FUAFTestBindableTraitSharedData::StaticStruct()->GetPathName();
+				const FString Properties = FString::Printf(
+					TEXT("(DecoratorSharedDataStruct=\"/Script/CoreUObject.ScriptStruct'%s'\",Name=\"%s\")"),
+					*SharedDataPath, *TestTraitName);
+
+				const FName TraitName = Controller->AddTrait(
+					FName(*TraitStackNodeName),
+					FName(TEXT("/Script/AnimNextAnimGraph.RigDecorator_AnimNextCppDecorator")),
+					FName(*TestTraitName),
+					Properties);
+				ASSERT_THAT(IsTrue(TraitName.ToString().Equals(TestTraitName)));
+			})
+			.Then(TEXT("Verify bTestBool pin has ConstantValue sub-pin with expected default"), [&]()
+			{
+				URigVMNode* StackNode = ControllerGraph->FindNodeByName(FName(*TraitStackNodeName));
+				ASSERT_THAT(IsNotNull(StackNode));
+
+				const FString BoolPinPath = FString::Printf(TEXT("%s.%s"), *TestTraitName, *BoolPinName);
+				URigVMPin* BoolPin = StackNode->FindPin(BoolPinPath);
+				ASSERT_THAT(IsNotNull(BoolPin));
+
+				// Latent property decomposition creates sub-pins
+				ASSERT_THAT(IsTrue(BoolPin->GetSubPins().Num() > 0));
+
+				const FString ConstantValuePath = FString::Printf(TEXT("%s.%s.%s"), *TestTraitName, *BoolPinName, *ConstantValueName);
+				URigVMPin* ConstantValuePin = StackNode->FindPin(ConstantValuePath);
+				ASSERT_THAT(IsNotNull(ConstantValuePin));
+
+				ASSERT_THAT(IsTrue(ConstantValuePin->GetDefaultValue() == TEXT("False")));
+			})
+			.Then(TEXT("Set ConstantValue to True via controller and verify round-trip"), [&]()
+			{
+				URigVMNode* StackNode = ControllerGraph->FindNodeByName(FName(*TraitStackNodeName));
+				const FString ConstantValuePath = FString::Printf(TEXT("%s.%s.%s"), *TestTraitName, *BoolPinName, *ConstantValueName);
+				URigVMPin* ConstantValuePin = StackNode->FindPin(ConstantValuePath);
+				ASSERT_THAT(IsNotNull(ConstantValuePin));
+
+				const FString ExpectedText = TEXT("True");
+
+				ASSERT_THAT(IsTrue(Controller->SetPinDefaultValue(ConstantValuePin->GetPinPath(), ExpectedText)));
+				ASSERT_THAT(IsTrue(ConstantValuePin->GetDefaultValue() == ExpectedText));
+			})
+			.Then(TEXT("Set ConstantValue back to False via controller and verify round-trip"), [&]()
+			{
+				URigVMNode* StackNode = ControllerGraph->FindNodeByName(FName(*TraitStackNodeName));
+				const FString ConstantValuePath = FString::Printf(TEXT("%s.%s.%s"), *TestTraitName, *BoolPinName, *ConstantValueName);
+				URigVMPin* ConstantValuePin = StackNode->FindPin(ConstantValuePath);
+				ASSERT_THAT(IsNotNull(ConstantValuePin));
+
+				const FString ExpectedText = TEXT("False");
+
+				ASSERT_THAT(IsTrue(Controller->SetPinDefaultValue(ConstantValuePin->GetPinPath(), ExpectedText)));
+				ASSERT_THAT(IsTrue(ConstantValuePin->GetDefaultValue() == ExpectedText));
+			})
+			.Then(TEXT("Verify parent pin text reflects sub-pin changes"), [&]()
+			{
+				URigVMNode* StackNode = ControllerGraph->FindNodeByName(FName(*TraitStackNodeName));
+
+				// Set ConstantValue to True so the parent pin text changes
+				const FString ConstantValuePath = FString::Printf(TEXT("%s.%s.%s"), *TestTraitName, *BoolPinName, *ConstantValueName);
+				URigVMPin* ConstantValuePin = StackNode->FindPin(ConstantValuePath);
+				ASSERT_THAT(IsNotNull(ConstantValuePin));
+				Controller->SetPinDefaultValue(ConstantValuePin->GetPinPath(), TEXT("True"));
+
+				const FString BoolPinPath = FString::Printf(TEXT("%s.%s"), *TestTraitName, *BoolPinName);
+				URigVMPin* BoolPin = StackNode->FindPin(BoolPinPath);
+				ASSERT_THAT(IsNotNull(BoolPin));
+
+				const FString ParentValue = BoolPin->GetDefaultValue();
+				UE_LOG(LogTemp, Display, TEXT("DIAG: ParentValue = '%s'"), *ParentValue);
+				UE_LOG(LogTemp, Display, TEXT("DIAG: ConstantValue after set = '%s'"), *ConstantValuePin->GetDefaultValue());
+				ASSERT_THAT(IsTrue(ParentValue.Contains(TEXT("True"))));
+			});
+	}
+
+	// Verifies that FBindableBool binding state round-trips through text
+	// export/import. Exercises SetBinding/ClearBinding with UScriptStruct
+	// ExportText/ImportText -- the same serialization path the RigVM uses
+	// for pin default values and the editor customization triggers via
+	// NotifyPostChange.
+	//
+	// This test proves that:
+	//  - Binding data survives ExportText -> ImportText round-trip
+	//  - HasBinding() and GetBinding() return the correct state after import
+	//  - Each distinct binding state produces unique export text
+	TEST_METHOD(FBindableBool_BindingTextSerializationRoundTrip)
+	{
+		TestCommandBuilder
+			.Do(TEXT("Export, import, and verify binding round-trip for each state"), [&]()
+			{
+				const UScriptStruct* BoolStruct = FBindableBool::StaticStruct();
+				ASSERT_THAT(IsNotNull(BoolStruct));
+
+				FUAFPropertyBinding TestBinding;
+				TestBinding.SourceType = EUAFBindingSourceType::Variable;
+
+				// --- State 1: bound (ConstantValue=true, binding set) ---
+				FBindableBool Bound(true);
+				Bound.SetBinding(TestBinding);
+
+				FString BoundText;
+				BoolStruct->ExportText(BoundText, &Bound, nullptr, nullptr, 0, nullptr, true);
+				ASSERT_THAT(IsTrue(BoundText.Contains(TEXT("__Binding="))));
+				ASSERT_THAT(IsFalse(BoundText.Contains(TEXT("__Binding=()"))));
+
+				// Import bound text into a fresh instance via UScriptStruct::ImportText
+				// (dispatches to native ImportTextItem for STRUCT_ImportTextItemNative)
+				FBindableBool ImportedBound;
+				ASSERT_THAT(IsNotNull(BoolStruct->ImportText(*BoundText, &ImportedBound, nullptr, 0, nullptr, BoolStruct->GetName())));
+				ASSERT_THAT(IsTrue(ImportedBound.HasBinding()));
+				ASSERT_THAT(IsTrue(ImportedBound.GetConstantValue() == true));
+
+				// Verify the imported binding matches the original
+				const FUAFPropertyBinding* ImportedBinding = ImportedBound.GetBinding();
+				ASSERT_THAT(IsNotNull(ImportedBinding));
+				ASSERT_THAT(IsTrue(ImportedBinding->SourceType == EUAFBindingSourceType::Variable));
+
+				// --- State 2: cleared (ConstantValue=true, binding cleared) ---
+				FBindableBool Cleared(true);
+				Cleared.SetBinding(TestBinding);
+				Cleared.ClearBinding();
+
+				FString ClearedText;
+				BoolStruct->ExportText(ClearedText, &Cleared, nullptr, nullptr, 0, nullptr, true);
+				ASSERT_THAT(IsTrue(ClearedText.Contains(TEXT("__Binding=()"))));
+
+				// Import cleared text
+				FBindableBool ImportedCleared;
+				ASSERT_THAT(IsNotNull(BoolStruct->ImportText(*ClearedText, &ImportedCleared, nullptr, 0, nullptr, BoolStruct->GetName())));
+				ASSERT_THAT(IsFalse(ImportedCleared.HasBinding()));
+				ASSERT_THAT(IsTrue(ImportedCleared.GetConstantValue() == true));
+
+				// --- State 3: unbound (ConstantValue=true, never had binding) ---
+				FBindableBool Unbound(true);
+				FString UnboundText;
+				BoolStruct->ExportText(UnboundText, &Unbound, nullptr, nullptr, 0, nullptr, true);
+
+				// Each binding state must produce unique export text so that
+				// the RigVM controller's text comparison detects every transition
+				ASSERT_THAT(IsTrue(BoundText != UnboundText));
+				ASSERT_THAT(IsTrue(BoundText != ClearedText));
+			});
+	}
+
+protected:
+	void CreateAnimNextAnimationGraph()
+	{
+		const TSubclassOf<UUAFAnimGraphFactory> AnimGraphFactoryClass = UUAFAnimGraphFactory::StaticClass();
+		UObject* FactoryObject = UAFTestsUtilities::CreateFactoryObject(NewObject<UUAFAnimGraphFactory>(GetTransientPackage(), AnimGraphFactoryClass.Get()), UUAFAnimGraph::StaticClass(), "TestBindableTraitPinGraph");
+		AnimNextAnimationGraph = CastChecked<UUAFAnimGraph>(FactoryObject);
+		ASSERT_THAT(IsNotNull(AnimNextAnimationGraph));
+		ASSERT_THAT(IsTrue(AssetEditorSubsystem->OpenEditorForAsset(AnimNextAnimationGraph)));
+	}
+
+	void CreateTraitStackNode()
+	{
+		Controller = UUAFTestBlueprintLibrary::GetControllerByName(AnimNextAnimationGraph, TEXT("RigVMGraph"));
+		ASSERT_THAT(IsNotNull(Controller));
+
+		ControllerGraph = Controller->GetGraph();
+		ASSERT_THAT(IsNotNull(ControllerGraph));
+
+		TestRunner->AddExpectedError(TEXT("Base Trait Data is Invalid. Please, select a new Base Trait."), EAutomationExpectedErrorFlags::Contains, -1);
+
+		URigVMUnitNode* TraitStackNode = Controller->AddUnitNodeFromStructPath(
+			FRigUnit_AnimNextTraitStack::StaticStruct()->GetPathName(),
+			TEXT("Execute"),
+			FVector2D::ZeroVector,
+			TEXT("AnimNextTraitStack"));
+		ASSERT_THAT(IsNotNull(TraitStackNode));
+		ASSERT_THAT(IsNotNull(ControllerGraph->FindNodeByName(FName(*TraitStackNodeName))));
+	}
+};
+
+#endif // WITH_EDITOR
+#endif // WITH_DEV_AUTOMATION_TESTS
diff --git a/Engine/Plugins/Experimental/UAF/UAFTestSuites/Source/UAFCQTestSuite/UAFCQTestSuite.Build.cs b/Engine/Plugins/Experimental/UAF/UAFTestSuites/Source/UAFCQTestSuite/UAFCQTestSuite.Build.cs
--- a/Engine/Plugins/Experimental/UAF/UAFTestSuites/Source/UAFCQTestSuite/UAFCQTestSuite.Build.cs
+++ b/Engine/Plugins/Experimental/UAF/UAFTestSuites/Source/UAFCQTestSuite/UAFCQTestSuite.Build.cs
@@ -50,8 +50,9 @@ public class UAFCQTestSuite : ModuleRules
 				"UAFAnimGraphEditor",
 				"UAFUncookedOnly",
 				"RigVM",
-				"RigVMDeveloper"
-				// ... add private dependencies that you statically link with here ...	
+				"RigVMDeveloper",
+				"UAFAnimNodeTestData"
+				// ... add private dependencies that you statically link with here ...
 			}
 		);
 
